Hot-reload from IDE

#change UBT: Added support for hot-reload when UBT is called from IDE and the current target is curring
#change Added support for hot-reload from VS to the engine
#change Refactored hot-reload functionality to a separarte module
#change Added hot-reload analytics events

diff --git a/Engine/Source/Developer/HotReload/HotReload.Build.cs b/Engine/Source/Developer/HotReload/HotReload.Build.cs
new file mode 100644
--- /dev/null
+++ b/Engine/Source/Developer/HotReload/HotReload.Build.cs
@@ -0,0 +1,36 @@
+// Copyright 1998-2014 Epic Games, Inc. All Rights Reserved.
+using UnrealBuildTool;
+
+public class HotReload : ModuleRules
+{
+	public HotReload(TargetInfo Target)
+	{
+			PublicDependencyModuleNames.AddRange(
+				new string[]
+				{
+					"Core",
+					// ... add other public dependencies that you statically link with here ...
+				}
+				);
+
+			PrivateDependencyModuleNames.AddRange(
+				new string[]
+				{
+					"Analytics",
+					"Core",
+					"CoreUObject",
+					"DirectoryWatcher",
+					// ... add other public dependencies that you statically link with here ...
+				}
+				);
+
+
+			PrivateDependencyModuleNames.AddRange(
+				new string[] 
+				{
+					"Engine",
+					"UnrealEd", 
+				}
+			);
+	}
+}
\ No newline at end of file
diff --git a/Engine/Source/Developer/HotReload/Private/HotReload.cpp b/Engine/Source/Developer/HotReload/Private/HotReload.cpp
new file mode 100644
--- /dev/null
+++ b/Engine/Source/Developer/HotReload/Private/HotReload.cpp
@@ -0,0 +1,589 @@
+// Copyright 1998-2014 Epic Games, Inc. All Rights Reserved.
+
+#include "HotReloadPrivatePCH.h"
+#include "Editor/UnrealEd/Public/Kismet2/KismetReinstanceUtilities.h"
+#include "Runtime/Analytics/Analytics/Public/Interfaces/IAnalyticsProvider.h"
+#include "EngineAnalytics.h"
+#include "ScopedTimers.h"
+
+DEFINE_LOG_CATEGORY(LogHotReload);
+
+#define LOCTEXT_NAMESPACE "HotReload"
+
+/**
+ * Module for HotReload support
+ */
+class FHotReloadModule : public IHotReloadModule
+{
+public:
+
+	/** IModuleInterface implementation */
+	virtual void StartupModule() override;
+	virtual void ShutdownModule() override;
+
+	/** IHotReloadInterface implementation */
+	virtual void AddHotReloadFunctionRemap(Native NewFunctionPointer, Native OldFunctionPointer) override;	
+	virtual void RebindPackages(TArray< UPackage* > Packages, TArray< FName > DependentModules, const bool bWaitForCompletion, FOutputDevice &Ar) override;
+	virtual void DoHotReloadFromEditor() override;
+	virtual FHotReloadEvent& OnHotReload() { return HotReloadEvent; }	
+
+private:
+	
+	/**
+	 * Adds a callback to directory watcher for the game binaries folder.
+	 */
+	void InitHotReloadWatcher();
+
+	/**
+	 * Removes a directory watcher callback
+	 */
+	void ShutdownHotReloadWatcher();
+
+	/**
+	 * Performs hot-reload from IDE (when game DLLs change)
+	 */
+	void DoHotReloadFromIDE();
+
+	/**
+	* Performs internal module recompilation
+	*/
+	ECompilationResult::Type RebindPackagesInternal(TArray< UPackage* > Packages, TArray< FName > DependentModules, const bool bWaitForCompletion, FOutputDevice &Ar);
+
+	/**
+	 * Does the actual hot-reload, unloads old modules, loads new ones
+	 */
+	ECompilationResult::Type DoHotReloadInternal(bool bRecompileFinished, bool bRecompileSucceeded, TArray<UPackage*> Packages, TArray<FName> InDependentModules, FOutputDevice &HotReloadAr);
+
+	/**
+	* Callback for async ompilation
+	*/
+	void DoHotReloadCallback(bool bRecompileFinished, bool bRecompileSucceeded, TArray<UPackage*> Packages, TArray<FName> InDependentModules, FOutputDevice &HotReloadAr);
+
+	/**
+	 * Gets all currently loaded game module names.
+	 */
+	void GetGameModules(TArray<FString>& OutGameModules);
+
+	/**
+	 * Gets packages to re-bind and dependent modules.
+	 */
+	void GetPackagesToRebindAndDependentModules(const TArray<FString>& InGameModuleNames, TArray<UPackage*>& OutPackagesToRebind, TArray<FName>& OutDependentModules);
+
+	/**
+	 * Called from CoreUObject to re-instance hot-reloaded classes
+	 */
+	void ReinstanceClass(UClass* OldClass, UClass* NewClass);
+
+	/**
+	 * Tick function for FTicker: checks for re-loaded modules and does hot-reload from IDE
+	 */
+	bool Tick(float DeltaTime);
+
+	/**
+	 * Directory watcher callback
+	 */
+	void OnHotReloadBinariesChanged(const TArray<struct FFileChangeData>& FileChanges);
+
+	/**	Broadcasts that a hot reload just finished. */
+	void BroadcastHotReload() 
+	{ 
+		HotReloadEvent.Broadcast(); 
+	}
+
+	/**
+	 * Sends analytics event about the re-load
+	*/
+	static void RecordAnalyticsEvent(const TCHAR* ReloadFrom, ECompilationResult::Type Result, double Duration, int32 PackageCount, int32 DependentModulesCount);
+
+	/** FTicker delegate (hot-reload from IDE) */
+	FTickerDelegate TickerDelegate;
+
+	/** Callback when game binaries folder changes */
+	IDirectoryWatcher::FDirectoryChanged BinariesFolderChangedDelegate;
+
+	/** True if currently hot-reloading from editor (suppresses hot-reload from IDE) */
+	bool bIsHotReloadingFromEditor;
+	
+	struct FRecompiledModule
+	{
+		FString Name;
+		FString NewFilename;
+		FRecompiledModule() {}
+		FRecompiledModule(const FString& InName, const FString& InFilename)
+			: Name(InName)
+			, NewFilename(InFilename)
+		{}
+	};
+
+	/** New module DLLs */
+	TArray<FRecompiledModule> NewModules;
+
+	/** Delegate broadcast when a module has been hot-reloaded */
+	FHotReloadEvent HotReloadEvent;
+};
+
+IMPLEMENT_MODULE(FHotReloadModule, HotReload);
+
+void FHotReloadModule::StartupModule()
+{
+	bIsHotReloadingFromEditor = false;
+
+	// Register re-instancing delegate (Core)
+	FCoreUObjectDelegates::ReplaceHotReloadClassDelegate.BindRaw(this, &FHotReloadModule::ReinstanceClass);
+	
+	// Register directory watcher delegate
+	InitHotReloadWatcher();
+
+	// Register hot-reload from IDE ticker
+	TickerDelegate = FTickerDelegate::CreateRaw(this, &FHotReloadModule::Tick);
+	FTicker::GetCoreTicker().AddTicker(TickerDelegate);
+}
+
+void FHotReloadModule::ShutdownModule()
+{
+	FTicker::GetCoreTicker().RemoveTicker(TickerDelegate);
+	ShutdownHotReloadWatcher();
+}
+
+/** Type hash for a UObject Function Pointer, maybe not a great choice, but it should be sufficient for the needs here. **/
+inline uint32 GetTypeHash(Native A)
+{
+	return *(uint32*)&A;
+}
+
+/** Map from old function pointer to new function pointer for hot reload. */
+static TMap<Native, Native> HotReloadFunctionRemap;
+
+/** Adds and entry for the UFunction native pointer remap table */
+void FHotReloadModule::AddHotReloadFunctionRemap(Native NewFunctionPointer, Native OldFunctionPointer)
+{
+	Native OtherNewFunction = HotReloadFunctionRemap.FindRef(OldFunctionPointer);
+	check(!OtherNewFunction || OtherNewFunction == NewFunctionPointer);
+	check(NewFunctionPointer);
+	check(OldFunctionPointer);
+	HotReloadFunctionRemap.Add(OldFunctionPointer, NewFunctionPointer);
+}
+
+void FHotReloadModule::DoHotReloadFromEditor()
+{
+	TArray<FString> GameModuleNames;
+	TArray<UPackage*> PackagesToRebind;
+	TArray<FName> DependentModules;
+	GetGameModules(GameModuleNames);
+	ECompilationResult::Type Result = ECompilationResult::Unsupported;
+	// Analytics
+	double Duration = 0.0;
+
+	if (GameModuleNames.Num() > 0)
+	{
+		FScopedDurationTimer Timer(Duration);
+
+
+		GetPackagesToRebindAndDependentModules(GameModuleNames, PackagesToRebind, DependentModules);
+
+		GWarn->BeginSlowTask(LOCTEXT("CompilingGameCode", "Compiling Game Code"), true);
+
+		const bool bWaitForCompletion = false;	// Don't wait -- we want compiling to happen asynchronously
+		Result = RebindPackagesInternal(PackagesToRebind, DependentModules, bWaitForCompletion, *GLog);
+
+		GWarn->EndSlowTask();
+
+		BroadcastHotReload();
+	}
+
+	RecordAnalyticsEvent(TEXT("Editor"), Result, Duration, PackagesToRebind.Num(), DependentModules.Num());
+}
+
+void FHotReloadModule::DoHotReloadCallback(bool bRecompileFinished, bool bRecompileSucceeded, TArray<UPackage*> Packages, TArray< FName > InDependentModules, FOutputDevice &HotReloadAr)
+{
+	DoHotReloadInternal(bRecompileFinished, bRecompileSucceeded, Packages, InDependentModules, HotReloadAr);
+}
+
+ECompilationResult::Type FHotReloadModule::DoHotReloadInternal(bool bRecompileFinished, bool bRecompileSucceeded, TArray<UPackage*> Packages, TArray< FName > InDependentModules, FOutputDevice &HotReloadAr)
+{
+	ECompilationResult::Type Result = ECompilationResult::Unsupported;
+#if !IS_MONOLITHIC
+	if (bRecompileSucceeded)
+	{
+		FFeedbackContext& ErrorsFC = UClass::GetDefaultPropertiesFeedbackContext();
+		ErrorsFC.Errors.Empty();
+		ErrorsFC.Warnings.Empty();
+		// Rebind the hot reload DLL 
+		TGuardValue<bool> GuardIsHotReload(GIsHotReload, true);
+		TGuardValue<bool> GuardIsInitialLoad(GIsInitialLoad, true);
+		HotReloadFunctionRemap.Empty(); // redundant
+
+		CollectGarbage(GARBAGE_COLLECTION_KEEPFLAGS); // we create a new CDO in the transient package...this needs to go away before we try again.
+
+		// Load the new modules up
+		bool bReloadSucceeded = false;
+		for (TArray<UPackage*>::TConstIterator CurPackageIt(Packages); CurPackageIt; ++CurPackageIt)
+		{
+			UPackage* Package = *CurPackageIt;
+			FName ShortPackageName = FPackageName::GetShortFName(Package->GetFName());
+
+			// Abandon the old module.  We can't unload it because various data structures may be living
+			// that have vtables pointing to code that would become invalidated.
+			FModuleManager::Get().AbandonModule(ShortPackageName);
+
+			// Module should never be loaded at this point
+			check(!FModuleManager::Get().IsModuleLoaded(ShortPackageName));
+
+			// Load the newly-recompiled module up (it will actually have a different DLL file name at this point.)
+			FModuleManager::Get().LoadModule(ShortPackageName);
+			bReloadSucceeded = FModuleManager::Get().IsModuleLoaded(ShortPackageName);
+			if (!bReloadSucceeded)
+			{
+				HotReloadAr.Logf(ELogVerbosity::Warning, TEXT("HotReload failed, reload failed %s."), *Package->GetName());
+				Result = ECompilationResult::OtherCompilationError;
+				break;
+			}
+		}
+		// Load dependent modules.
+		for (int32 Nx = 0; Nx < InDependentModules.Num(); ++Nx)
+		{
+			const FName ModuleName = InDependentModules[Nx];
+
+			FModuleManager::Get().UnloadOrAbandonModuleWithCallback(ModuleName, HotReloadAr);
+			const bool bLoaded = FModuleManager::Get().LoadModuleWithCallback(ModuleName, HotReloadAr);
+			if (!bLoaded)
+			{
+				HotReloadAr.Logf(ELogVerbosity::Warning, TEXT("Unable to reload module %s"), *ModuleName.GetPlainNameString());
+			}
+		}
+
+		if (ErrorsFC.Errors.Num() || ErrorsFC.Warnings.Num())
+		{
+			TArray<FString> All;
+			All = ErrorsFC.Errors;
+			All += ErrorsFC.Warnings;
+
+			ErrorsFC.Errors.Empty();
+			ErrorsFC.Warnings.Empty();
+
+			FString AllInOne;
+			for (int32 Index = 0; Index < All.Num(); Index++)
+			{
+				AllInOne += All[Index];
+				AllInOne += TEXT("\n");
+			}
+			HotReloadAr.Logf(ELogVerbosity::Warning, TEXT("Some classes could not be reloaded:\n%s"), *AllInOne);
+		}
+
+		if (bReloadSucceeded)
+		{
+			int32 Count = 0;
+			// Remap all native functions (and gather scriptstructs)
+			TArray<UScriptStruct*> ScriptStructs;
+			for (FRawObjectIterator It; It; ++It)
+			{
+				if (UFunction* Function = Cast<UFunction>(*It))
+				{
+					if (Native NewFunction = HotReloadFunctionRemap.FindRef(Function->GetNativeFunc()))
+					{
+						Count++;
+						Function->SetNativeFunc(NewFunction);
+					}
+				}
+
+				if (UScriptStruct* ScriptStruct = Cast<UScriptStruct>(*It))
+				{
+					if (Packages.ContainsByPredicate([=](UPackage* Package) { return ScriptStruct->IsIn(Package); }) && ScriptStruct->GetCppStructOps())
+					{
+						ScriptStructs.Add(ScriptStruct);
+					}
+				}
+			}
+			// now let's set up the script structs...this relies on super behavior, so null them all, then set them all up. Internally this sets them up hierarchically.
+			for (int32 ScriptIndex = 0; ScriptIndex < ScriptStructs.Num(); ScriptIndex++)
+			{
+				ScriptStructs[ScriptIndex]->ClearCppStructOps();
+			}
+			for (int32 ScriptIndex = 0; ScriptIndex < ScriptStructs.Num(); ScriptIndex++)
+			{
+				ScriptStructs[ScriptIndex]->PrepareCppStructOps();
+				check(ScriptStructs[ScriptIndex]->GetCppStructOps());
+			}
+			HotReloadAr.Logf(ELogVerbosity::Warning, TEXT("HotReload successful (%d functions remapped  %d scriptstructs remapped)"), Count, ScriptStructs.Num());
+
+			HotReloadFunctionRemap.Empty();
+			Result = ECompilationResult::Succeeded;
+		}
+	}
+	else if (bRecompileFinished)
+	{
+		HotReloadAr.Logf(ELogVerbosity::Warning, TEXT("HotReload failed, recompile failed"));
+		Result = ECompilationResult::OtherCompilationError;
+	}
+#endif
+	bIsHotReloadingFromEditor = false;
+	return Result;
+}
+
+void FHotReloadModule::RebindPackages(TArray<UPackage*> InPackages, TArray<FName> DependentModules, const bool bWaitForCompletion, FOutputDevice &Ar)
+{
+	ECompilationResult::Type Result = ECompilationResult::Unknown;
+	double Duration = 0.0;
+	{
+		FScopedDurationTimer RebindTimer(Duration);
+		Result = RebindPackagesInternal(InPackages, DependentModules, bWaitForCompletion, Ar);
+	}
+	RecordAnalyticsEvent(TEXT("Rebind"), Result, Duration, InPackages.Num(), DependentModules.Num());
+}
+
+ECompilationResult::Type FHotReloadModule::RebindPackagesInternal(TArray<UPackage*> InPackages, TArray<FName> DependentModules, const bool bWaitForCompletion, FOutputDevice &Ar)
+{
+	ECompilationResult::Type Result = ECompilationResult::Unsupported;
+#if !IS_MONOLITHIC
+	bool bCanRebind = InPackages.Num() > 0;
+
+	// Verify that we're going to be able to rebind the specified packages
+	if (bCanRebind)
+	{
+		for (UPackage* Package : InPackages)
+		{
+			check(Package);
+
+			if (Package->GetOuter() != NULL)
+			{
+				Ar.Logf(ELogVerbosity::Warning, TEXT("Could not rebind package for %s, package is either not bound yet or is not a DLL."), *Package->GetName());
+				bCanRebind = false;
+				break;
+			}
+		}
+	}
+
+	// We can only proceed if a compile isn't already in progress
+	if (FModuleManager::Get().IsCurrentlyCompiling())
+	{
+		Ar.Logf(ELogVerbosity::Warning, TEXT("Could not rebind package because a module compile is already in progress."));
+		bCanRebind = false;
+	}
+
+	if (bCanRebind)
+	{
+		bIsHotReloadingFromEditor = true;
+
+		const double StartTime = FPlatformTime::Seconds();
+
+		TArray< FName > ModuleNames;
+		for (UPackage* Package : InPackages)
+		{
+			// Attempt to recompile this package's module
+			FName ShortPackageName = FPackageName::GetShortFName(Package->GetFName());
+			ModuleNames.Add(ShortPackageName);
+		}
+
+		// Add dependent modules.
+		ModuleNames.Append(DependentModules);
+
+		// Start compiling modules
+		const bool bCompileStarted = FModuleManager::Get().RecompileModulesAsync(
+			ModuleNames,
+			FModuleManager::FRecompileModulesCallback::CreateRaw< FHotReloadModule, TArray<UPackage*>, TArray<FName>, FOutputDevice& >(this, &FHotReloadModule::DoHotReloadCallback, InPackages, DependentModules, Ar),
+			bWaitForCompletion,
+			Ar);
+
+		if (bCompileStarted)
+		{
+			if (bWaitForCompletion)
+			{
+				Ar.Logf(ELogVerbosity::Warning, TEXT("HotReload operation took %4.1fs."), float(FPlatformTime::Seconds() - StartTime));
+			}
+			else
+			{
+				Ar.Logf(ELogVerbosity::Warning, TEXT("Starting HotReload took %4.1fs."), float(FPlatformTime::Seconds() - StartTime));
+			}
+			Result = ECompilationResult::Succeeded;
+		}
+		else
+		{
+			Ar.Logf(ELogVerbosity::Warning, TEXT("RebindPackages failed because the compiler could not be started."));
+			Result = ECompilationResult::OtherCompilationError;
+		}
+	}
+	else
+#endif
+	{
+		Ar.Logf(ELogVerbosity::Warning, TEXT("RebindPackages not possible for specified packages (or application was compiled in monolithic mode.)"));
+	}
+	return Result;
+}
+
+void FHotReloadModule::ReinstanceClass(UClass* OldClass, UClass* NewClass)
+{
+	UE_LOG(LogHotReload, Log, TEXT("Re-instancing %s after hot-reload."), *NewClass->GetName());
+	FBlueprintCompileReinstancer ReinstanceHelper(NewClass, OldClass);
+	ReinstanceHelper.ReinstanceObjects();
+}
+
+void FHotReloadModule::GetGameModules(TArray<FString>& OutGameModules)
+{
+	// Ask the module manager for a list of currently-loaded gameplay modules
+	TArray< FModuleStatus > ModuleStatuses;
+	FModuleManager::Get().QueryModules(ModuleStatuses);
+
+	for (TArray< FModuleStatus >::TConstIterator ModuleStatusIt(ModuleStatuses); ModuleStatusIt; ++ModuleStatusIt)
+	{
+		const FModuleStatus& ModuleStatus = *ModuleStatusIt;
+
+		// We only care about game modules that are currently loaded
+		if (ModuleStatus.bIsLoaded && ModuleStatus.bIsGameModule)
+		{
+			OutGameModules.AddUnique(ModuleStatus.Name);
+		}
+	}
+}
+
+void FHotReloadModule::OnHotReloadBinariesChanged(const TArray<struct FFileChangeData>& FileChanges)
+{
+	if (bIsHotReloadingFromEditor)
+	{
+		// DO NOTHING, this case is handled by RebindPackages
+		return;
+	}
+
+	TArray< FString > GameModuleNames;
+	GetGameModules(GameModuleNames);
+
+	if (GameModuleNames.Num() > 0)
+	{
+		// Check if any of the game DLLs has been added
+		for (auto& Change : FileChanges)
+		{
+			if (Change.Action == FFileChangeData::FCA_Added)
+			{
+				const FString Filename = FPaths::GetCleanFilename(Change.Filename);
+				if (Filename.EndsWith(FPlatformProcess::GetModuleExtension()))
+				{
+					for (auto& GameModule : GameModuleNames)
+					{
+						if (Filename.Contains(GameModule) && !NewModules.ContainsByPredicate([&](const FRecompiledModule& Module){ return Module.Name == GameModule; }))
+						{
+							// Add to queue. We do not hot-reload here as there may potentially be other modules being compiled.
+							NewModules.Add(FRecompiledModule(GameModule, Change.Filename));
+							UE_LOG(LogHotReload, Log, TEXT("New module detected: %s"), *Filename);
+						}
+					}
+				}
+			}
+		}
+	}
+}
+
+void FHotReloadModule::InitHotReloadWatcher()
+{
+	FDirectoryWatcherModule& DirectoryWatcherModule = FModuleManager::Get().LoadModuleChecked<FDirectoryWatcherModule>(TEXT("DirectoryWatcher"));
+	IDirectoryWatcher* DirectoryWatcher = DirectoryWatcherModule.Get();
+	if (DirectoryWatcher)
+	{
+		// Watch the game binaries folder for new files
+		FString BinariesPath = FPaths::ConvertRelativePathToFull(FPaths::GameDir() / TEXT("Binaries") / FPlatformProcess::GetBinariesSubdirectory());
+		BinariesFolderChangedDelegate = IDirectoryWatcher::FDirectoryChanged::CreateRaw(this, &FHotReloadModule::OnHotReloadBinariesChanged);
+		DirectoryWatcher->RegisterDirectoryChangedCallback(BinariesPath, BinariesFolderChangedDelegate);
+	}
+}
+
+void FHotReloadModule::ShutdownHotReloadWatcher()
+{
+	FDirectoryWatcherModule& DirectoryWatcherModule = FModuleManager::Get().LoadModuleChecked<FDirectoryWatcherModule>(TEXT("DirectoryWatcher"));
+	IDirectoryWatcher* DirectoryWatcher = DirectoryWatcherModule.Get();
+	if (DirectoryWatcher)
+	{
+		FString BinariesPath = FPaths::ConvertRelativePathToFull(FPaths::GameDir() / TEXT("Binaries") / FPlatformProcess::GetBinariesSubdirectory());
+		DirectoryWatcher->UnregisterDirectoryChangedCallback(BinariesPath, BinariesFolderChangedDelegate);
+	}
+}
+
+bool FHotReloadModule::Tick(float DeltaTime)
+{
+	if (NewModules.Num())
+	{
+		// We have new modules in the queue, but make sure UBT has finished compiling all of them
+		if (!FPlatformProcess::IsApplicationRunning(TEXT("UnrealBuildTool")))
+		{
+			DoHotReloadFromIDE();
+			NewModules.Empty();
+		}
+		else
+		{
+			UE_LOG(LogHotReload, Verbose, TEXT("Detected %d reloaded modules but UnrealBuildTool is still running"), NewModules.Num());
+		}
+	}
+	return true;
+}
+
+void FHotReloadModule::GetPackagesToRebindAndDependentModules(const TArray<FString>& InGameModuleNames, TArray<UPackage*>& OutPackagesToRebind, TArray<FName>& OutDependentModules)
+{
+	for (auto& GameModuleName : InGameModuleNames)
+	{
+		FString PackagePath(FString(TEXT("/Script/")) + GameModuleName);
+		UPackage* Package = FindPackage(NULL, *PackagePath);
+		if (Package != NULL)
+		{
+			OutPackagesToRebind.Add(Package);
+		}
+		else
+		{
+			OutDependentModules.Add(*GameModuleName);
+		}
+	}
+}
+
+void FHotReloadModule::DoHotReloadFromIDE()
+{
+	TArray<FString> GameModuleNames;
+	TArray<UPackage*> PackagesToRebind;
+	TArray<FName> DependentModules;
+	double Duration = 0.0;
+	ECompilationResult::Type Result = ECompilationResult::Unsupported;
+
+	GetGameModules(GameModuleNames);
+
+	if (GameModuleNames.Num() > 0)
+	{		
+		FScopedDurationTimer Timer(Duration);
+
+		UE_LOG(LogHotReload, Log, TEXT("Starting Hot-Reload from IDE"));
+
+		GWarn->BeginSlowTask(LOCTEXT("CompilingGameCode", "Compiling Game Code"), true);
+
+		// Update compile data before we start compiling
+		for (auto& NewModule : NewModules)
+		{
+			FModuleManager::Get().UpdateModuleCompileData(*NewModule.Name);
+			FModuleManager::Get().OnModuleCompileSucceeded(*NewModule.Name, NewModule.NewFilename);
+		}
+
+
+		GetPackagesToRebindAndDependentModules(GameModuleNames, PackagesToRebind, DependentModules);
+
+		check(PackagesToRebind.Num() || DependentModules.Num())
+		{
+			const bool bRecompileFinished = true;
+			const bool bRecompileSucceeded = true;
+			Result = DoHotReloadInternal(bRecompileFinished, bRecompileSucceeded, PackagesToRebind, DependentModules, *GLog);
+		}
+
+		GWarn->EndSlowTask();
+		BroadcastHotReload();		
+	}
+
+	RecordAnalyticsEvent(TEXT("IDE"), Result, Duration, PackagesToRebind.Num(), DependentModules.Num());
+}
+
+void FHotReloadModule::RecordAnalyticsEvent(const TCHAR* ReloadFrom, ECompilationResult::Type Result, double Duration, int32 PackageCount, int32 DependentModulesCount)
+{
+	if (FEngineAnalytics::IsAvailable())
+	{
+		TArray< FAnalyticsEventAttribute > ReloadAttribs;
+		ReloadAttribs.Add(FAnalyticsEventAttribute(TEXT("ReloadFrom"), ReloadFrom));
+		ReloadAttribs.Add(FAnalyticsEventAttribute(TEXT("Result"), ECompilationResult::ToString(Result)));
+		ReloadAttribs.Add(FAnalyticsEventAttribute(TEXT("Duration"), FString::Printf(TEXT("%.4lf"), Duration)));
+		ReloadAttribs.Add(FAnalyticsEventAttribute(TEXT("Packages"), FString::Printf(TEXT("%d"), PackageCount)));
+		ReloadAttribs.Add(FAnalyticsEventAttribute(TEXT("DependentModules"), FString::Printf(TEXT("%d"), DependentModulesCount)));
+		FEngineAnalytics::GetProvider().RecordEvent(TEXT("Editor.Usage.HotReload"), ReloadAttribs);
+	}
+}
+
+#undef LOCTEXT_NAMESPACE
\ No newline at end of file
diff --git a/Engine/Source/Developer/HotReload/Private/HotReloadPrivatePCH.h b/Engine/Source/Developer/HotReload/Private/HotReloadPrivatePCH.h
new file mode 100644
--- /dev/null
+++ b/Engine/Source/Developer/HotReload/Private/HotReloadPrivatePCH.h
@@ -0,0 +1,9 @@
+// Copyright 1998-2014 Epic Games, Inc. All Rights Reserved.
+
+#pragma once
+
+#include "CoreUObject.h"
+#include "IHotReload.h"
+#include "Developer/DirectoryWatcher/Public/DirectoryWatcherModule.h"
+
+DECLARE_LOG_CATEGORY_EXTERN(LogHotReload, Log, All);
diff --git a/Engine/Source/Developer/HotReload/Public/IHotReload.h b/Engine/Source/Developer/HotReload/Public/IHotReload.h
new file mode 100644
--- /dev/null
+++ b/Engine/Source/Developer/HotReload/Public/IHotReload.h
@@ -0,0 +1,37 @@
+// Copyright 1998-2014 Epic Games, Inc. All Rights Reserved.
+
+#pragma once
+
+#include "HotReloadInterface.h"
+
+/**
+* The public interface to this module.  In most cases, this interface is only public to sibling modules
+* within this plugin.
+*/
+class IHotReloadModule : public IHotReloadInterface
+{
+
+public:
+
+	/**
+	* Singleton-like access to this module's interface.  This is just for convenience!
+	* Beware of calling this during the shutdown phase, though.  Your module might have been unloaded already.
+	*
+	* @return Returns singleton instance, loading the module on demand if needed
+	*/
+	static inline IHotReloadModule& Get()
+	{
+		return FModuleManager::LoadModuleChecked<IHotReloadModule>("HotReload");
+	}
+
+	/**
+	* Checks to see if this module is loaded and ready.  It is only valid to call Get() if IsAvailable() returns true.
+	*
+	* @return True if the module is loaded and ready to use
+	*/
+	static inline bool IsAvailable()
+	{
+		return FModuleManager::Get().IsModuleLoaded("HotReload");
+	}
+};
+
diff --git a/Engine/Source/Developer/ModuleUI/Private/SModuleUI.cpp b/Engine/Source/Developer/ModuleUI/Private/SModuleUI.cpp
--- a/Engine/Source/Developer/ModuleUI/Private/SModuleUI.cpp
+++ b/Engine/Source/Developer/ModuleUI/Private/SModuleUI.cpp
@@ -1,7 +1,7 @@
 // Copyright 1998-2014 Epic Games, Inc. All Rights Reserved.
 
 #include "ModuleUIPrivatePCH.h"
-
+#include "HotReloadInterface.h"
 
 BEGIN_SLATE_FUNCTION_BUILD_OPTIMIZATION
 void SModuleUI::Construct(const SModuleUI::FArguments& InArgs)
@@ -255,8 +255,9 @@ FReply SModuleUI::FModuleListItem::OnRecompileClicked()
 		if( PackagesToRebind.Num() > 0 )
 		{
 			// Perform a hot reload
-			const bool bWaitForCompletion = true;
-			RebindPackages( PackagesToRebind, TArray<FName>(), bWaitForCompletion, *GLog );
+			const bool bWaitForCompletion = true;			
+			IHotReloadInterface& HotReloadSupport = FModuleManager::LoadModuleChecked<IHotReloadInterface>("HotReload");
+			HotReloadSupport.RebindPackages(PackagesToRebind, TArray<FName>(), bWaitForCompletion, *GLog);
 		}
 		else
 		{
diff --git a/Engine/Source/Editor/BehaviorTreeEditor/Private/BehaviorTreeEditorTypes.cpp b/Engine/Source/Editor/BehaviorTreeEditor/Private/BehaviorTreeEditorTypes.cpp
--- a/Engine/Source/Editor/BehaviorTreeEditor/Private/BehaviorTreeEditorTypes.cpp
+++ b/Engine/Source/Editor/BehaviorTreeEditor/Private/BehaviorTreeEditorTypes.cpp
@@ -8,6 +8,7 @@
 #include "BehaviorTree/Tasks/BTTask_BlueprintBase.h"
 #include "BehaviorTree/Decorators/BTDecorator_BlueprintBase.h"
 #include "BehaviorTree/Services/BTService_BlueprintBase.h"
+#include "HotReloadInterface.h"
 
 const FString UBehaviorTreeEditorTypes::PinCategory_MultipleNodes("MultipleNodes");
 const FString UBehaviorTreeEditorTypes::PinCategory_SingleComposite("SingleComposite");
@@ -143,8 +144,11 @@ FClassBrowseHelper::FClassBrowseHelper()
 	AssetRegistryModule.Get().OnAssetAdded().AddRaw( this, &FClassBrowseHelper::OnAssetAdded);
 	AssetRegistryModule.Get().OnAssetRemoved().AddRaw( this, &FClassBrowseHelper::OnAssetRemoved );
 
-	// Register to have Populate called when doing a Hot Reload or when a Blueprint is compiled.
-	GEditor->OnHotReload().AddRaw( this, &FClassBrowseHelper::InvalidateCache );
+	// Register to have Populate called when doing a Hot Reload.
+	IHotReloadInterface& HotReloadSupport = FModuleManager::LoadModuleChecked<IHotReloadInterface>("HotReload");
+	HotReloadSupport.OnHotReload().AddRaw(this, &FClassBrowseHelper::InvalidateCache);
+
+	// Register to have Populate called when a Blueprint is compiled.
 	GEditor->OnBlueprintCompiled().AddRaw( this, &FClassBrowseHelper::InvalidateCache );
 	GEditor->OnClassPackageLoadedOrUnloaded().AddRaw( this, &FClassBrowseHelper::InvalidateCache );
 
@@ -161,8 +165,11 @@ FClassBrowseHelper::~FClassBrowseHelper()
 		AssetRegistryModule.Get().OnAssetAdded().RemoveAll(this);
 		AssetRegistryModule.Get().OnAssetRemoved().RemoveAll(this);
 
-		// Unregister to have Populate called when doing a Hot Reload or when a Blueprint is compiled.
-		GEditor->OnHotReload().RemoveAll(this);
+		// Unregister to have Populate called when doing a Hot Reload.
+		IHotReloadInterface& HotReloadSupport = FModuleManager::LoadModuleChecked<IHotReloadInterface>("HotReload");
+		HotReloadSupport.OnHotReload().RemoveAll(this);
+
+		// Unregister to have Populate called when a Blueprint is compiled.
 		GEditor->OnBlueprintCompiled().RemoveAll(this);
 		GEditor->OnClassPackageLoadedOrUnloaded().RemoveAll(this);
 	}
diff --git a/Engine/Source/Editor/ClassViewer/Private/SClassViewer.cpp b/Engine/Source/Editor/ClassViewer/Private/SClassViewer.cpp
--- a/Engine/Source/Editor/ClassViewer/Private/SClassViewer.cpp
+++ b/Engine/Source/Editor/ClassViewer/Private/SClassViewer.cpp
@@ -30,6 +30,7 @@
 #include "GameProjectGenerationModule.h"
 
 #include "SourceCodeNavigation.h"
+#include "HotReloadInterface.h"
 
 #define LOCTEXT_NAMESPACE "SClassViewer"
 
@@ -1584,8 +1585,11 @@ FClassHierarchy::FClassHierarchy()
 	AssetRegistryModule.Get().OnAssetAdded().AddRaw( this, &FClassHierarchy::AddAsset);
 	AssetRegistryModule.Get().OnAssetRemoved().AddRaw( this, &FClassHierarchy::RemoveAsset );
 
-	// Register to have Populate called when doing a Hot Reload or when a Blueprint is compiled.
-	GEditor->OnHotReload().AddStatic(ClassViewer::Helpers::RequestPopulateClassHierarchy);
+	// Register to have Populate called when doing a Hot Reload.
+	IHotReloadInterface& HotReloadSupport = FModuleManager::LoadModuleChecked<IHotReloadInterface>("HotReload");
+	HotReloadSupport.OnHotReload().AddStatic(ClassViewer::Helpers::RequestPopulateClassHierarchy);
+
+	// Register to have Populate called when a Blueprint is compiled.
 	GEditor->OnBlueprintCompiled().AddStatic(ClassViewer::Helpers::RequestPopulateClassHierarchy);
 	GEditor->OnClassPackageLoadedOrUnloaded().AddStatic(ClassViewer::Helpers::RequestPopulateClassHierarchy);
 
@@ -1602,8 +1606,11 @@ FClassHierarchy::~FClassHierarchy()
 		AssetRegistryModule.Get().OnAssetAdded().RemoveAll( this );
 		AssetRegistryModule.Get().OnAssetRemoved().RemoveAll( this );
 
-		// Unregister to have Populate called when doing a Hot Reload or when a Blueprint is compiled.
-		GEditor->OnHotReload().RemoveStatic(ClassViewer::Helpers::RequestPopulateClassHierarchy);
+		// Unregister to have Populate called when doing a Hot Reload.
+		IHotReloadInterface& HotReloadSupport = FModuleManager::LoadModuleChecked<IHotReloadInterface>("HotReload");
+		HotReloadSupport.OnHotReload().RemoveStatic(ClassViewer::Helpers::RequestPopulateClassHierarchy);
+
+		// Unregister to have Populate called when a Blueprint is compiled.
 		GEditor->OnBlueprintCompiled().RemoveStatic(ClassViewer::Helpers::RequestPopulateClassHierarchy);
 		GEditor->OnClassPackageLoadedOrUnloaded().RemoveStatic(ClassViewer::Helpers::RequestPopulateClassHierarchy);
 	}
diff --git a/Engine/Source/Editor/LevelEditor/Private/LevelEditorActions.cpp b/Engine/Source/Editor/LevelEditor/Private/LevelEditorActions.cpp
--- a/Engine/Source/Editor/LevelEditor/Private/LevelEditorActions.cpp
+++ b/Engine/Source/Editor/LevelEditor/Private/LevelEditorActions.cpp
@@ -48,6 +48,7 @@
 #include "EditorActorFolders.h"
 #include "ActorPickerMode.h"
 #include "EngineBuildSettings.h"
+#include "HotReloadInterface.h"
 
 DEFINE_LOG_CATEGORY_STATIC(LevelEditorActions, Log, All);
 
@@ -899,55 +900,8 @@ void FLevelEditorActionCallbacks::RecompileGameCode_Clicked()
 	// Don't allow a recompile while already compiling!
 	if( !FModuleManager::Get().IsCurrentlyCompiling() )
 	{
-		TArray< FString > GameModuleNames;
-		{
-			// Ask the module manager for a list of currently-loaded gameplay modules
-			TArray< FModuleStatus > ModuleStatuses;
-			FModuleManager::Get().QueryModules( ModuleStatuses );
-
-			for( TArray< FModuleStatus >::TConstIterator ModuleStatusIt( ModuleStatuses ); ModuleStatusIt; ++ModuleStatusIt )
-			{
-				const FModuleStatus& ModuleStatus = *ModuleStatusIt;
-
-				// We only care about game modules that are currently loaded
-				if( ModuleStatus.bIsLoaded && ModuleStatus.bIsGameModule )
-				{
-					GameModuleNames.AddUnique( ModuleStatus.Name );
-				}
-			}
-		}
-
-		if( GameModuleNames.Num() > 0 )
-		{
-			TArray< UPackage*> PackagesToRebind;
-			// Packages that rely on packages with UObject, but don't contain UObjects.
-			TArray< FName > DependentModules;
-
-			for( int32 CurModuleIndex = 0; CurModuleIndex < GameModuleNames.Num(); ++CurModuleIndex )
-			{
-				const FString& GameModuleName = *GameModuleNames[ CurModuleIndex ];
-				FString PackagePath( FString( TEXT( "/Script/" ) ) + GameModuleName );
-
-				UPackage* Package = FindPackage(NULL, *PackagePath);
-				if( Package != NULL )
-				{
-					PackagesToRebind.Add( Package );
-				}
-				else
-				{
-					DependentModules.Add( *GameModuleName );
-				}
-			}
-
-			GWarn->BeginSlowTask( LOCTEXT("CompilingGameCode", "Compiling Game Code"), true );
-
-			const bool bWaitForCompletion = false;	// Don't wait -- we want compiling to happen asynchronously
-			RebindPackages( PackagesToRebind, DependentModules, bWaitForCompletion, *GLog );
-
-			GWarn->EndSlowTask();
-
-			GEditor->BroadcastHotReload();
-		}
+		IHotReloadInterface& HotReloadSupport = FModuleManager::LoadModuleChecked<IHotReloadInterface>("HotReload");
+		HotReloadSupport.DoHotReloadFromEditor();
 	}
 }
 
diff --git a/Engine/Source/Editor/UnrealEd/Classes/Editor/EditorEngine.h b/Engine/Source/Editor/UnrealEd/Classes/Editor/EditorEngine.h
--- a/Engine/Source/Editor/UnrealEd/Classes/Editor/EditorEngine.h
+++ b/Engine/Source/Editor/UnrealEd/Classes/Editor/EditorEngine.h
@@ -506,13 +506,6 @@ public:
 	/** Annotation to track which PIE/SIE (PlayWorld) UObjects have counterparts in the EditorWorld **/
 	class FUObjectAnnotationSparseBool ObjectsThatExistInEditorWorld;
 
-	/** Called when a Hot Reload event has completed. */
-	DECLARE_EVENT( UEditorEngine, FHotReloadEvent );
-	FHotReloadEvent& OnHotReload() { return HotReloadEvent; }
-
-	/**	Broadcasts that a hot reload just finished. THIS SHOULD NOT BE PUBLIC */
-	void BroadcastHotReload() { HotReloadEvent.Broadcast(); }
-
 	/** Called when a Blueprint compile is completed. */
 	DECLARE_EVENT( UEditorEngine, FBlueprintCompiledEvent );
 	FBlueprintCompiledEvent& OnBlueprintCompiled() { return BlueprintCompiledEvent; }
@@ -2494,9 +2487,6 @@ private:
 
 private:
 
-	/** Delegate broadcast when a module has been hot-reloaded */
-	FHotReloadEvent HotReloadEvent;
-
 	/** Delegate broadcast when blueprint is compiled */
 	FBlueprintCompiledEvent BlueprintCompiledEvent;
 
diff --git a/Engine/Source/Editor/UnrealEd/Private/Editor.cpp b/Engine/Source/Editor/UnrealEd/Private/Editor.cpp
--- a/Engine/Source/Editor/UnrealEd/Private/Editor.cpp
+++ b/Engine/Source/Editor/UnrealEd/Private/Editor.cpp
@@ -584,6 +584,8 @@ void UEditorEngine::Init(IEngineLoop* InEngineLoop)
 		{
 			FModuleManager::Get().LoadModule(TEXT("GameplayAbilitiesEditor"));
 		}
+
+		FModuleManager::Get().LoadModule(TEXT("HotReload"));
 	}
 
 	float BSPTexelScale = 100.0f;
diff --git a/Engine/Source/Editor/UnrealEd/Private/UnrealEd.cpp b/Engine/Source/Editor/UnrealEd/Private/UnrealEd.cpp
--- a/Engine/Source/Editor/UnrealEd/Private/UnrealEd.cpp
+++ b/Engine/Source/Editor/UnrealEd/Private/UnrealEd.cpp
@@ -22,7 +22,6 @@
 #include "AssetEditorManager.h"
 
 #include "EditorActorFolders.h"
-#include "Kismet2/KismetReinstanceUtilities.h"
 
 UUnrealEdEngine* GUnrealEd;
 
@@ -187,22 +186,4 @@ void EditorExit()
 
 }
 
-/** Singleton used for class re-instancing during hot-reload */
-struct FUClassHotReloadReinstancer
-{
-	FUClassHotReloadReinstancer()
-	{
-		FCoreUObjectDelegates::ReplaceHotReloadClassDelegate.BindStatic(&FUClassHotReloadReinstancer::ReinstanceClass);
-	}
-
-	static void ReinstanceClass(UClass* OldClass, UClass* NewClass)
-	{
-		UE_LOG(LogUnrealEd, Display, TEXT("Re-instancing %s after hot-reload."), *NewClass->GetName());
-		FBlueprintCompileReinstancer ReinstanceHelper(NewClass, OldClass);
-		ReinstanceHelper.ReinstanceObjects();
-	}
-};
-static FUClassHotReloadReinstancer GHotReloadReinstancer;
-
-
 IMPLEMENT_MODULE( FDefaultModuleImpl, UnrealEd );
diff --git a/Engine/Source/Editor/UnrealEd/UnrealEd.Build.cs b/Engine/Source/Editor/UnrealEd/UnrealEd.Build.cs
--- a/Engine/Source/Editor/UnrealEd/UnrealEd.Build.cs
+++ b/Engine/Source/Editor/UnrealEd/UnrealEd.Build.cs
@@ -176,7 +176,8 @@ public class UnrealEd : ModuleRules
 				"UndoHistory",
 				"SourceCodeAccess",
 				"ReferenceViewer",
-				"EditorLiveStreaming"
+				"EditorLiveStreaming",
+				"HotReload",
 			}
 		);
 
diff --git a/Engine/Source/Programs/UnrealBuildTool/Configuration/UEBuildBinary.cs b/Engine/Source/Programs/UnrealBuildTool/Configuration/UEBuildBinary.cs
--- a/Engine/Source/Programs/UnrealBuildTool/Configuration/UEBuildBinary.cs
+++ b/Engine/Source/Programs/UnrealBuildTool/Configuration/UEBuildBinary.cs
@@ -217,6 +217,16 @@ namespace UnrealBuildTool
             return null;
         }
 
+		/// <summary>
+		/// Called to allow the binary to find game modules.
+		/// </summary>
+		/// <param name="OnlyModules"></param>
+		/// <returns>The OnlyModule if found, null if not</returns>
+		public virtual List<UEBuildModule> FindGameModules()
+		{
+			return null;
+		}
+
 		/// <summary>
 		/// Generates a list of all modules referenced by this binary
 		/// </summary>
@@ -676,6 +686,20 @@ namespace UnrealBuildTool
 			return null;
 		}
 
+		public override List<UEBuildModule> FindGameModules()
+		{
+			var GameModules = new List<UEBuildModule>();
+			foreach (var ModuleName in ModuleNames)
+			{
+				UEBuildModule Module = Target.GetModuleByName(ModuleName);
+				if (!Utils.IsFileUnderDirectory(Module.ModuleDirectory, BuildConfiguration.RelativeEnginePath))
+				{
+					GameModules.Add(Module);
+				}
+			}
+			return GameModules;
+		}
+
 		// Object interface.
 
 		/// <summary>
diff --git a/Engine/Source/Programs/UnrealBuildTool/Configuration/UEBuildClient.cs b/Engine/Source/Programs/UnrealBuildTool/Configuration/UEBuildClient.cs
--- a/Engine/Source/Programs/UnrealBuildTool/Configuration/UEBuildClient.cs
+++ b/Engine/Source/Programs/UnrealBuildTool/Configuration/UEBuildClient.cs
@@ -16,7 +16,8 @@ namespace UnrealBuildTool
 			TargetRules InRulesObject,
 			List<string> InAdditionalDefinitions, 
 			string InRemoteRoot, 
-			List<OnlyModule> InOnlyModules)
+			List<OnlyModule> InOnlyModules,
+			bool bInEditorRecompile)
 			// NOTE: If we're building a monolithic binary, then the game and engine code are linked together into one
 			//       program executable, so we want the application name to be the game name.  In the case of a modular
 			//       binary, we use 'UnrealEngine' for our application name
@@ -28,7 +29,9 @@ namespace UnrealBuildTool
                 InRulesObject: InRulesObject,
                 InAdditionalDefinitions: InAdditionalDefinitions,
                 InRemoteRoot: InRemoteRoot,
-                InOnlyModules: InOnlyModules)
+                InOnlyModules: InOnlyModules,
+				bInEditorRecompile: bInEditorRecompile
+			)
         {
             if (ShouldCompileMonolithic())
             {
diff --git a/Engine/Source/Programs/UnrealBuildTool/Configuration/UEBuildConfiguration.cs b/Engine/Source/Programs/UnrealBuildTool/Configuration/UEBuildConfiguration.cs
--- a/Engine/Source/Programs/UnrealBuildTool/Configuration/UEBuildConfiguration.cs
+++ b/Engine/Source/Programs/UnrealBuildTool/Configuration/UEBuildConfiguration.cs
@@ -172,6 +172,12 @@ namespace UnrealBuildTool
 		[XmlConfig]
 		public static bool bCompileForSize;
 
+		/** True if hot-reload from IDE is allowed */
+		public static bool bAllowHotReloadFromIDE;
+
+		/** True if performing hot-reload from IDE */
+		public static bool bHotReloadFromIDE;
+
 		/// <summary>
 		/// Sets the configuration back to defaults.
 		/// </summary>
@@ -208,6 +214,8 @@ namespace UnrealBuildTool
 			bCompilePhysXVehicle = true;
 			bCompileFreeType = true;
 			bCompileForSize = false;
+			bHotReloadFromIDE = false;
+			bAllowHotReloadFromIDE = true;
 		}
 
 		/// <summary>
diff --git a/Engine/Source/Programs/UnrealBuildTool/Configuration/UEBuildEditor.cs b/Engine/Source/Programs/UnrealBuildTool/Configuration/UEBuildEditor.cs
--- a/Engine/Source/Programs/UnrealBuildTool/Configuration/UEBuildEditor.cs
+++ b/Engine/Source/Programs/UnrealBuildTool/Configuration/UEBuildEditor.cs
@@ -16,7 +16,8 @@ namespace UnrealBuildTool
 			TargetRules InRulesObject,
 			List<string> InAdditionalDefinitions, 
 			string InRemoteRoot, 
-			List<OnlyModule> InOnlyModules)
+			List<OnlyModule> InOnlyModules,
+			bool bInEditorRecompile)
 			// NOTE: If we're building a monolithic binary, then the game and engine code are linked together into one
 			//       program executable, so we want the application name to be the game name.  In the case of a modular
 			//       binary, we use 'UnrealEngine' for our application name
@@ -34,7 +35,9 @@ namespace UnrealBuildTool
 				InRulesObject: InRulesObject, 
 				InAdditionalDefinitions:InAdditionalDefinitions, 
 				InRemoteRoot:InRemoteRoot, 
-				InOnlyModules:InOnlyModules)
+				InOnlyModules:InOnlyModules,
+				bInEditorRecompile: bInEditorRecompile
+			)
 		{
 		}
 
diff --git a/Engine/Source/Programs/UnrealBuildTool/Configuration/UEBuildGame.cs b/Engine/Source/Programs/UnrealBuildTool/Configuration/UEBuildGame.cs
--- a/Engine/Source/Programs/UnrealBuildTool/Configuration/UEBuildGame.cs
+++ b/Engine/Source/Programs/UnrealBuildTool/Configuration/UEBuildGame.cs
@@ -16,7 +16,8 @@ namespace UnrealBuildTool
 			TargetRules InRulesObject,
 			List<string> InAdditionalDefinitions, 
 			string InRemoteRoot, 
-			List<OnlyModule> InOnlyModules)
+			List<OnlyModule> InOnlyModules,
+			bool bInEditorRecompile)
 			// NOTE: If we're building a monolithic binary, then the game and engine code are linked together into one
 			//       program executable, so we want the application name to be the game name.  In the case of a modular
 			//       binary, we use 'UnrealEngine' for our application name
@@ -28,7 +29,9 @@ namespace UnrealBuildTool
 				InRulesObject: InRulesObject, 
 				InAdditionalDefinitions:InAdditionalDefinitions,
 				InRemoteRoot:InRemoteRoot,
-				InOnlyModules:InOnlyModules)
+				InOnlyModules:InOnlyModules,
+				bInEditorRecompile: bInEditorRecompile
+			)
 		{
 			if (ShouldCompileMonolithic())
 			{
diff --git a/Engine/Source/Programs/UnrealBuildTool/Configuration/UEBuildServer.cs b/Engine/Source/Programs/UnrealBuildTool/Configuration/UEBuildServer.cs
--- a/Engine/Source/Programs/UnrealBuildTool/Configuration/UEBuildServer.cs
+++ b/Engine/Source/Programs/UnrealBuildTool/Configuration/UEBuildServer.cs
@@ -16,7 +16,8 @@ namespace UnrealBuildTool
 			TargetRules InRulesObject,
 			List<string> InAdditionalDefinitions, 
 			string InRemoteRoot, 
-			List<OnlyModule> InOnlyModules)
+			List<OnlyModule> InOnlyModules,
+			bool bInEditorRecompile)
 			// NOTE: If we're building a monolithic binary, then the game and engine code are linked together into one
 			//       program executable, so we want the application name to be the game name.  In the case of a modular
 			//       binary, we use 'UnrealEngine' for our application name
@@ -28,7 +29,9 @@ namespace UnrealBuildTool
 				InRulesObject: InRulesObject, 
 				InAdditionalDefinitions:InAdditionalDefinitions,
 				InRemoteRoot:InRemoteRoot,
-				InOnlyModules:InOnlyModules)
+				InOnlyModules:InOnlyModules,
+				bInEditorRecompile: bInEditorRecompile
+			)
 		{
 		}
 
diff --git a/Engine/Source/Programs/UnrealBuildTool/Configuration/UEBuildTarget.cs b/Engine/Source/Programs/UnrealBuildTool/Configuration/UEBuildTarget.cs
--- a/Engine/Source/Programs/UnrealBuildTool/Configuration/UEBuildTarget.cs
+++ b/Engine/Source/Programs/UnrealBuildTool/Configuration/UEBuildTarget.cs
@@ -611,6 +611,9 @@ namespace UnrealBuildTool
 		/** Extra engine module names to either include in the binary (monolithic) or create side-by-side DLLs for (modular) */
 		public List<string> ExtraModuleNames = new List<string>();
 
+		/** True if re-compiling this target from the editor */
+		public bool bEditorRecompile;
+
 		/** If building only a specific set of modules, these are the modules to build */
 		protected List<OnlyModule> OnlyModules = new List<OnlyModule>();
 
@@ -645,13 +648,15 @@ namespace UnrealBuildTool
 			TargetRules InRulesObject,
 			List<string> InAdditionalDefinitions,
 			string InRemoteRoot,
-			List<OnlyModule> InOnlyModules)
+			List<OnlyModule> InOnlyModules,
+			bool bInEditorRecompile)
 		{
 			AppName = InAppName;
 			GameName = InGameName;
 			Platform = InPlatform;
 			Configuration = InConfiguration;
 			Rules = InRulesObject;
+			bEditorRecompile = bInEditorRecompile;
 
 			{
 				bCompileMonolithic = (Rules != null) ? Rules.ShouldCompileMonolithic(InPlatform, InConfiguration) : false;
@@ -1536,39 +1541,11 @@ namespace UnrealBuildTool
 			// If we're building a single module, then find the binary for that module and add it to our target
 			if (OnlyModules.Count > 0)
 			{
-				var FilteredBinaries = new List<UEBuildBinary>();
-
-				var AnyBinariesAdded = false;
-				foreach (var DLLBinary in AppBinaries)
-				{
-					var FoundOnlyModule = DLLBinary.FindOnlyModule(OnlyModules);
-					if (FoundOnlyModule != null)
-					{
-						FilteredBinaries.Add(DLLBinary);
-						AnyBinariesAdded = true;
-
-						if (!String.IsNullOrEmpty(FoundOnlyModule.OnlyModuleSuffix))
-						{
-							var Appendage = "-" + FoundOnlyModule.OnlyModuleSuffix;
-
-							var MatchPos = DLLBinary.Config.OutputFilePath.LastIndexOf(FoundOnlyModule.OnlyModuleName, StringComparison.InvariantCultureIgnoreCase);
-							if (MatchPos < 0)
-							{
-								throw new BuildException("Failed to find module name \"{0}\" specified on the command line inside of the output filename \"{1}\" to add appendage.", FoundOnlyModule.OnlyModuleName, DLLBinary.Config.OutputFilePath);
-							}
-							DLLBinary.Config.OriginalOutputFilePath = DLLBinary.Config.OutputFilePath;
-							DLLBinary.Config.OutputFilePath = DLLBinary.Config.OutputFilePath.Insert(MatchPos + FoundOnlyModule.OnlyModuleName.Length, Appendage);
-						}
-					}
-				}
-
-				// Copy the result into the final list
-				AppBinaries = FilteredBinaries;
-
-				if (!AnyBinariesAdded)
-				{
-					throw new BuildException("One or more of the modules specified using the '-module' argument could not be found.");
-				}
+				AppBinaries = FilterOnlyModules();
+			}
+			else if (UEBuildConfiguration.bHotReloadFromIDE)
+			{
+				AppBinaries = FilterGameModules();
 			}
 
 			// Filter out binaries that were already built and are just used for linking. We will not build these binaries but we need them for link information
@@ -1655,6 +1632,76 @@ namespace UnrealBuildTool
 
 			return SpecialRocketLibFilesThatAreBuildProducts;
 		}
+
+		private static string AddModuleFilenameSuffix(string ModuleName, string FilePath, string Suffix)
+		{			
+			var MatchPos = FilePath.LastIndexOf(ModuleName, StringComparison.InvariantCultureIgnoreCase);
+			if (MatchPos < 0)
+			{
+				throw new BuildException("Failed to find module name \"{0}\" specified on the command line inside of the output filename \"{1}\" to add appendage.", ModuleName, FilePath);
+			}
+			var Appendage = "-" + Suffix;
+			return FilePath.Insert(MatchPos + ModuleName.Length, Appendage);
+		}
+
+		private List<UEBuildBinary> FilterOnlyModules()
+		{
+			var FilteredBinaries = new List<UEBuildBinary>();
+
+			var AnyBinariesAdded = false;
+			foreach (var DLLBinary in AppBinaries)
+			{
+				var FoundOnlyModule = DLLBinary.FindOnlyModule(OnlyModules);
+				if (FoundOnlyModule != null)
+				{
+					FilteredBinaries.Add(DLLBinary);
+					AnyBinariesAdded = true;
+
+					if (!String.IsNullOrEmpty(FoundOnlyModule.OnlyModuleSuffix))
+					{
+						DLLBinary.Config.OriginalOutputFilePath = DLLBinary.Config.OutputFilePath;
+						DLLBinary.Config.OutputFilePath = AddModuleFilenameSuffix(FoundOnlyModule.OnlyModuleName, DLLBinary.Config.OutputFilePath, FoundOnlyModule.OnlyModuleSuffix);
+					}
+				}
+			}
+
+			if (!AnyBinariesAdded)
+			{
+				throw new BuildException("One or more of the modules specified using the '-module' argument could not be found.");
+			}
+
+			// Copy the result into the final list
+			return FilteredBinaries;			
+		}
+
+		private List<UEBuildBinary> FilterGameModules()
+		{
+			var FilteredBinaries = new List<UEBuildBinary>();
+
+			var AnyBinariesAdded = false;
+			foreach (var DLLBinary in AppBinaries)
+			{
+				var GameModules = DLLBinary.FindGameModules();
+				if (GameModules != null && GameModules.Count > 0)
+				{
+					var UniqueSuffix = (new Random((int)(DateTime.Now.Ticks % Int32.MaxValue)).Next(10000)).ToString();
+					DLLBinary.Config.OriginalOutputFilePath = DLLBinary.Config.OutputFilePath;
+					DLLBinary.Config.OutputFilePath = AddModuleFilenameSuffix(GameModules[0].Name, DLLBinary.Config.OutputFilePath, UniqueSuffix);
+
+					FilteredBinaries.Add(DLLBinary);
+					AnyBinariesAdded = true;
+				}
+			}
+
+			if (!AnyBinariesAdded)
+			{
+				throw new BuildException("One or more of the modules specified using the '-module' argument could not be found.");
+			}
+
+			// Copy the result into the final list
+			return FilteredBinaries;
+		}
+
 		/// <summary>
 		/// Writes target info.
 		/// </summary>
diff --git a/Engine/Source/Programs/UnrealBuildTool/System/ExternalExecution.cs b/Engine/Source/Programs/UnrealBuildTool/System/ExternalExecution.cs
--- a/Engine/Source/Programs/UnrealBuildTool/System/ExternalExecution.cs
+++ b/Engine/Source/Programs/UnrealBuildTool/System/ExternalExecution.cs
@@ -18,8 +18,16 @@ namespace UnrealBuildTool
 	// processes valid.
 	public enum ECompilationResult
 	{
+		/** Compilation succeeded */
 		Succeeded = 0,
-		OtherCompilationError = 2
+		/** Compilation failed because generated code changed which was not supported */
+		FailedDueToHeaderChange = 1,
+		/** Compilation failed due to compilation errors */
+		OtherCompilationError,
+		/** Compilation is not supported in the current build */
+		Unsupported,
+		/** Unknown error */
+		Unknown
 	}
 
 	/** Information about a module that needs to be passed to UnrealHeaderTool for code generation */
@@ -551,9 +559,10 @@ namespace UnrealBuildTool
 
 				if( !bIsBuildingUHT && bUHTNeedsToRun )
 				{
-					// Always build UnrealHeaderTool if header regeneration is required, unless we're running within a Rocket ecosystem
+					// Always build UnrealHeaderTool if header regeneration is required, unless we're running within a Rocket ecosystem or hot-reloading
 					if (UnrealBuildTool.RunningRocket() == false && 
 						UEBuildConfiguration.bDoNotBuildUHT == false &&
+						UEBuildConfiguration.bHotReloadFromIDE == false &&
 						!( !UnrealBuildTool.IsGatheringBuild && UnrealBuildTool.IsAssemblingBuild ) )	// If running in "assembler only" mode, we assume UHT is already up to date for much faster iteration!
 					{
 						// If it is out of date or not there it will be built.
diff --git a/Engine/Source/Programs/UnrealBuildTool/System/RulesCompiler.cs b/Engine/Source/Programs/UnrealBuildTool/System/RulesCompiler.cs
--- a/Engine/Source/Programs/UnrealBuildTool/System/RulesCompiler.cs
+++ b/Engine/Source/Programs/UnrealBuildTool/System/RulesCompiler.cs
@@ -1723,7 +1723,8 @@ namespace UnrealBuildTool
 							InRulesObject: RulesObject,
 							InAdditionalDefinitions: InAdditionalDefinitions, 
 							InRemoteRoot:InRemoteRoot, 
-							InOnlyModules:InOnlyModules);
+							InOnlyModules:InOnlyModules,
+							bInEditorRecompile: bInEditorRecompile);
 					}
 					break;
 				case TargetRules.TargetType.Editor:
@@ -1735,7 +1736,8 @@ namespace UnrealBuildTool
 							InRulesObject: RulesObject,
 							InAdditionalDefinitions: InAdditionalDefinitions,
 							InRemoteRoot: InRemoteRoot,
-							InOnlyModules: InOnlyModules);
+							InOnlyModules: InOnlyModules,
+							bInEditorRecompile: bInEditorRecompile);
 					}
 					break;
                 case TargetRules.TargetType.Client:
@@ -1747,7 +1749,8 @@ namespace UnrealBuildTool
                             InRulesObject: RulesObject,
                             InAdditionalDefinitions: InAdditionalDefinitions,
                             InRemoteRoot: InRemoteRoot,
-                            InOnlyModules: InOnlyModules);
+                            InOnlyModules: InOnlyModules,
+							bInEditorRecompile: bInEditorRecompile);
                     }
                     break;
 				case TargetRules.TargetType.Server:
@@ -1759,7 +1762,8 @@ namespace UnrealBuildTool
 							InRulesObject: RulesObject,
 							InAdditionalDefinitions: InAdditionalDefinitions,
 							InRemoteRoot: InRemoteRoot,
-							InOnlyModules: InOnlyModules);
+							InOnlyModules: InOnlyModules,
+							bInEditorRecompile: bInEditorRecompile);
 					}
 					break;
 				case TargetRules.TargetType.Program:
@@ -1772,22 +1776,10 @@ namespace UnrealBuildTool
 							InRulesObject:RulesObject, 
 							InAdditionalDefinitions:InAdditionalDefinitions, 
 							InRemoteRoot:InRemoteRoot, 
-							InOnlyModules:InOnlyModules);
+							InOnlyModules:InOnlyModules,
+							bInEditorRecompile:bInEditorRecompile);
 					}
 					break;
-					/*
-				case TargetRules.TargetType.RocketGame:
-					{
-						BuildTarget = new UEBuildRocketGame(
-							InGameName: TargetName,
-							InPlatform: Target.Platform,
-							InConfiguration: Target.Configuration,
-							InRulesObject: RulesObject,
-							InAdditionalDefinitions: InAdditionalDefinitions,
-							InRemoteRoot: InRemoteRoot,
-							InOnlyModules: InOnlyModules);
-					}
-					break;*/
 			}
 
 			if (BuildTarget == null)
diff --git a/Engine/Source/Programs/UnrealBuildTool/System/UnrealBuildTool.cs b/Engine/Source/Programs/UnrealBuildTool/System/UnrealBuildTool.cs
--- a/Engine/Source/Programs/UnrealBuildTool/System/UnrealBuildTool.cs
+++ b/Engine/Source/Programs/UnrealBuildTool/System/UnrealBuildTool.cs
@@ -580,6 +580,7 @@ namespace UnrealBuildTool
         private static int Main(string[] Arguments)
         {
             InitLogging();
+            Log.TraceVerbose(Environment.CommandLine);
 
             try
             {
@@ -1415,6 +1416,13 @@ namespace UnrealBuildTool
                 {
                     var TargetStartTime = DateTime.UtcNow;
 
+                    // Check if we're currently running an editor and this is hot-reload from IDE
+                    UEBuildConfiguration.bHotReloadFromIDE = UEBuildConfiguration.bAllowHotReloadFromIDE && !Target.bEditorRecompile && ShouldDoHotReload(Target);
+                    if (UEBuildConfiguration.bHotReloadFromIDE)
+                    {
+                        Log.TraceInformation("Compiling game modules for Hot-Reload.");
+                    }
+
                     // When in 'assembler only' mode, we'll load this cache later on a worker thread.  It takes a long time to load!
                     if( !( !UnrealBuildTool.IsGatheringBuild && UnrealBuildTool.IsAssemblingBuild ) )
                     {	
@@ -1708,6 +1716,27 @@ namespace UnrealBuildTool
             return BuildResult;
         }
 
+		/// <summary>
+		/// Checks if the editor is currently running and this is a hot-reload
+		/// </summary>
+		/// <param name="ResetPlatform"></param>
+		/// <param name="ResetConfiguration"></param>
+		/// <param name="Target"></param>
+		/// <returns></returns>
+		private static bool ShouldDoHotReload(UEBuildTarget Target)
+		{
+			bool bIsRunning = false;
+			if (!ProjectFileGenerator.bGenerateProjectFiles && !UEBuildConfiguration.bGenerateManifest && Target.Rules.Type == TargetRules.TargetType.Editor)
+			{
+				var EditorProcessFilename = UEBuildTarget.MakeBinaryPath("Launcher", "UE4Editor", Target.Platform, Target.Configuration, UEBuildBinaryType.Executable, null, false, null, "UE4Editor");
+				var EditorProcessName = Path.GetFileNameWithoutExtension(EditorProcessFilename);
+				var EditorProcesses = Process.GetProcessesByName(EditorProcessName);
+				var BinariesPath = Path.GetFullPath(Path.GetDirectoryName(EditorProcessFilename));
+				bIsRunning = EditorProcesses.FirstOrDefault(p => Path.GetFullPath(p.MainModule.FileName).StartsWith(BinariesPath, StringComparison.InvariantCultureIgnoreCase)) != default(Process);
+			}
+			return bIsRunning;
+		}
+
         private static void ParseBuildConfigurationFlags(string[] Arguments)
         {
             int ArgumentIndex = 0;
diff --git a/Engine/Source/Runtime/Core/Private/Modules/ModuleManager.cpp b/Engine/Source/Runtime/Core/Private/Modules/ModuleManager.cpp
--- a/Engine/Source/Runtime/Core/Private/Modules/ModuleManager.cpp
+++ b/Engine/Source/Runtime/Core/Private/Modules/ModuleManager.cpp
@@ -977,6 +977,15 @@ FString FModuleManager::MakeUBTArgumentsForModuleCompiling()
 	return AdditionalArguments;
 }
 
+void FModuleManager::OnModuleCompileSucceeded(FName ModuleName, const FString& NewModuleFilename)
+{
+	TSharedRef< FModuleInfo > ModuleInfo = Modules.FindChecked(ModuleName);
+
+	// If the compile succeeded, update the module info entry with the new file name for this module
+	ModuleInfo->Filename = NewModuleFilename;
+
+	OnModuleCompileSucceeded(ModuleName, ModuleInfo);
+}
 
 void FModuleManager::OnModuleCompileSucceeded(FName ModuleName, TSharedRef<FModuleManager::FModuleInfo> ModuleInfo)
 {
@@ -1004,6 +1013,11 @@ void FModuleManager::OnModuleCompileSucceeded(FName ModuleName, TSharedRef<FModu
 #endif
 }
 
+void FModuleManager::UpdateModuleCompileData(FName ModuleFName)
+{
+	TSharedRef< FModuleInfo > ModuleInfo = Modules.FindChecked(ModuleFName);
+	UpdateModuleCompileData(ModuleFName, ModuleInfo);
+}
 
 void FModuleManager::UpdateModuleCompileData(FName ModuleName, TSharedRef<FModuleManager::FModuleInfo> ModuleInfo)
 {
diff --git a/Engine/Source/Runtime/Core/Public/Modules/ModuleManager.h b/Engine/Source/Runtime/Core/Public/Modules/ModuleManager.h
--- a/Engine/Source/Runtime/Core/Public/Modules/ModuleManager.h
+++ b/Engine/Source/Runtime/Core/Public/Modules/ModuleManager.h
@@ -46,10 +46,36 @@ namespace ECompilationResult
 {
 	enum Type
 	{
+		/** Compilation succeeded */
 		Succeeded = 0,
+		/** Compilation failed because generated code changed which was not supported */
 		FailedDueToHeaderChange = 1,
-		OtherCompilationError = 2
+		/** Compilation failed due to compilation errors */
+		OtherCompilationError,
+		/** Compilation is not supported in the current build */
+		Unsupported,
+		/** Unknown error */
+		Unknown
 	};
+
+	/**
+	* Converts ECompilationResult enum to string.
+	*/
+	static FORCEINLINE const TCHAR* ToString(ECompilationResult::Type Result)
+	{
+		switch (Result)
+		{
+		case ECompilationResult::Succeeded:
+			return TEXT("Succeeded");
+		case ECompilationResult::FailedDueToHeaderChange:
+			return TEXT("FailedDueToHeaderChange");
+		case ECompilationResult::OtherCompilationError:
+			return TEXT("OtherCompilationError");
+		case ECompilationResult::Unsupported:
+			return TEXT("Unsupported");
+		};
+		return TEXT("Unknown");
+	}
 }
 
 
@@ -517,6 +543,12 @@ public:
 	 */
 	static const TCHAR *GetUBTConfiguration( );
 
+	/** Called for successfully re-complied module */
+	void OnModuleCompileSucceeded(FName ModuleName, const FString& NewModuleFilename);
+
+	/** Called when the compile data for a module need to be update in memory and written to config */
+	void UpdateModuleCompileData(FName ModuleName);
+
 public:
 
 	/**
diff --git a/Engine/Source/Runtime/CoreUObject/Private/UObject/Class.cpp b/Engine/Source/Runtime/CoreUObject/Private/UObject/Class.cpp
--- a/Engine/Source/Runtime/CoreUObject/Private/UObject/Class.cpp
+++ b/Engine/Source/Runtime/CoreUObject/Private/UObject/Class.cpp
@@ -5,8 +5,8 @@
 =============================================================================*/
 
 #include "CoreUObjectPrivate.h"
-
 #include "PropertyTag.h"
+#include "HotReloadInterface.h"
 
 DECLARE_LOG_CATEGORY_EXTERN(LogScriptSerialization, Log, All);
 DEFINE_LOG_CATEGORY(LogScriptSerialization);
@@ -3414,13 +3414,15 @@ void UClass::AddNativeFunction(const ANSICHAR* InName,Native InPointer)
 #if !IS_MONOLITHIC
 	if (GIsHotReload)
 	{
+		IHotReloadInterface& HotReloadSupport = FModuleManager::LoadModuleChecked<IHotReloadInterface>("HotReload");
+
 		// Find the function in the class's native function lookup table.
 		for(int32 FunctionIndex = 0;FunctionIndex < NativeFunctionLookupTable.Num();++FunctionIndex)
 		{
 			FNativeFunctionLookup& NativeFunctionLookup = NativeFunctionLookupTable[FunctionIndex];
 			if(NativeFunctionLookup.Name == InFName)
 			{
-				AddHotReloadFunctionRemap(InPointer, NativeFunctionLookup.Pointer);
+				HotReloadSupport.AddHotReloadFunctionRemap(InPointer, NativeFunctionLookup.Pointer);
 				NativeFunctionLookup.Pointer = InPointer;
 				return;
 			}
diff --git a/Engine/Source/Runtime/CoreUObject/Private/UObject/Obj.cpp b/Engine/Source/Runtime/CoreUObject/Private/UObject/Obj.cpp
--- a/Engine/Source/Runtime/CoreUObject/Private/UObject/Obj.cpp
+++ b/Engine/Source/Runtime/CoreUObject/Private/UObject/Obj.cpp
@@ -12,6 +12,7 @@
 
 #include "Serialization/ArchiveDescribeReference.h"
 #include "FindStronglyConnected.h"
+#include "HotReloadInterface.h"
 
 DEFINE_LOG_CATEGORY(LogObj);
 
@@ -1815,219 +1816,6 @@ void StaticShutdownAfterError()
 	}
 }
 
-
-/** Type hash for a UObject Function Pointer, maybe not a great choice, but it should be sufficient for the needs here. **/
-inline uint32 GetTypeHash( Native A )
-{
-	return *(uint32*)&A;
-}
-
-/** Map from old function pointer to new function pointer for hot reload. */
-static TMap<Native,Native> HotReloadFunctionRemap;
-
-/** Adds and entry for the UFunction native pointer remap table */
-void AddHotReloadFunctionRemap(Native NewFunctionPointer, Native OldFunctionPointer)
-{
-	Native OtherNewFunction = HotReloadFunctionRemap.FindRef(OldFunctionPointer);
-	check(!OtherNewFunction || OtherNewFunction == NewFunctionPointer);
-	check(NewFunctionPointer);
-	check(OldFunctionPointer);
-	HotReloadFunctionRemap.Add(OldFunctionPointer,NewFunctionPointer);
-}
-
-
-void RebindPackages( TArray<UPackage*> InPackages, TArray<FName> DependentModules, const bool bWaitForCompletion, FOutputDevice &Ar )
-{
-#if !IS_MONOLITHIC
-	bool bCanRebind = InPackages.Num() > 0;
-
-	// Verify that we're going to be able to rebind the specified packages
-	if( bCanRebind )
-	{
-		for( UPackage* Package : InPackages )
-		{
-			check(Package);
-
-			if (Package->GetOuter() != NULL )
-			{
-				Ar.Logf(ELogVerbosity::Warning, TEXT("Could not rebind package for %s, package is either not bound yet or is not a DLL."), *Package->GetName() );
-				bCanRebind = false;
-				break;
-			}
-		}
-	}
-
-
-	// We can only proceed if a compile isn't already in progress
-	if( FModuleManager::Get().IsCurrentlyCompiling() )
-	{
-		Ar.Logf(ELogVerbosity::Warning, TEXT("Could not rebind package because a module compile is already in progress.") );
-		bCanRebind = false;
-	}
-
-	if( bCanRebind )
-	{
-		struct Local
-		{
-			static void DoHotReload( bool bRecompileFinished, bool bRecompileSucceeded, TArray<UPackage*> Packages, TArray< FName > InDependentModules, FOutputDevice &HotReloadAr )
-			{
-				if( bRecompileSucceeded )
-				{
-					FFeedbackContext& ErrorsFC = UClass::GetDefaultPropertiesFeedbackContext();
-					ErrorsFC.Errors.Empty();
-					ErrorsFC.Warnings.Empty();
-					// Rebind the hot reload DLL 
-					TGuardValue<bool> GuardIsHotReload(GIsHotReload, true);
-					TGuardValue<bool> GuardIsInitialLoad(GIsInitialLoad, true);
-					HotReloadFunctionRemap.Empty(); // redundant
-
-					CollectGarbage(GARBAGE_COLLECTION_KEEPFLAGS); // we create a new CDO in the transient package...this needs to go away before we try again.
-
-					// Load the new modules up
-					bool bReloadSucceeded = false;
-					for( TArray<UPackage*>::TConstIterator CurPackageIt( Packages ); CurPackageIt; ++CurPackageIt )
-					{
-						UPackage* Package = *CurPackageIt;
-						FName ShortPackageName = FPackageName::GetShortFName(Package->GetFName());
-
-						// Abandon the old module.  We can't unload it because various data structures may be living
-						// that have vtables pointing to code that would become invalidated.
-						FModuleManager::Get().AbandonModule( ShortPackageName );
-
-						// Module should never be loaded at this point
-						check( !FModuleManager::Get().IsModuleLoaded( ShortPackageName ) );
-
-						// Load the newly-recompiled module up (it will actually have a different DLL file name at this point.)
-						FModuleManager::Get().LoadModule( ShortPackageName );
-						bReloadSucceeded = FModuleManager::Get().IsModuleLoaded( ShortPackageName );
-						if( !bReloadSucceeded )
-						{
-							HotReloadAr.Logf(ELogVerbosity::Warning, TEXT("HotReload failed, reload failed %s."),*Package->GetName());
-							break;
-						}
-					}
-					// Load dependent modules.
-					for( int32 Nx = 0; Nx < InDependentModules.Num(); ++Nx )
-					{
-						const FName ModuleName = InDependentModules[Nx];
-
-						FModuleManager::Get().UnloadOrAbandonModuleWithCallback( ModuleName, HotReloadAr );
-						const bool bLoaded = FModuleManager::Get().LoadModuleWithCallback( ModuleName, HotReloadAr );
-						if( !bLoaded )
-						{
-							HotReloadAr.Logf(ELogVerbosity::Warning, TEXT("Unable to reload module %s"), *ModuleName.GetPlainNameString());
-						}
-					}
-
-					if (ErrorsFC.Errors.Num() || ErrorsFC.Warnings.Num())
-					{
-						TArray<FString> All;
-						All = ErrorsFC.Errors;
-						All += ErrorsFC.Warnings;
-
-						ErrorsFC.Errors.Empty();
-						ErrorsFC.Warnings.Empty();
-
-						FString AllInOne;
-						for (int32 Index = 0; Index < All.Num(); Index++)
-						{
-							AllInOne += All[Index];
-							AllInOne += TEXT("\n");
-						}
-						HotReloadAr.Logf(ELogVerbosity::Warning, TEXT( "Some classes could not be reloaded:\n%s" ), *AllInOne );
-					}
-
-					if( bReloadSucceeded )
-					{
-						int32 Count = 0;
-						// Remap all native functions (and gather scriptstructs)
-						TArray<UScriptStruct*> ScriptStructs;
-						for ( FRawObjectIterator It; It; ++It )
-						{
-							if (UFunction* Function = Cast<UFunction>(*It))
-							{
-								if (Native NewFunction = HotReloadFunctionRemap.FindRef(Function->GetNativeFunc()))
-								{
-									Count++;
-									Function->SetNativeFunc(NewFunction);
-								}
-							}
-
-							if (UScriptStruct* ScriptStruct = Cast<UScriptStruct>(*It))
-							{
-								if (Packages.ContainsByPredicate([=](UPackage* Package) { return ScriptStruct->IsIn(Package); }) && ScriptStruct->GetCppStructOps())
-								{
-									ScriptStructs.Add(ScriptStruct);
-								}
-							}
-						}
-						// now let's set up the script structs...this relies on super behavior, so null them all, then set them all up. Internally this sets them up hierarchically.
-						for (int32 ScriptIndex = 0; ScriptIndex < ScriptStructs.Num(); ScriptIndex++)
-						{
-							ScriptStructs[ScriptIndex]->ClearCppStructOps();
-						}
-						for (int32 ScriptIndex = 0; ScriptIndex < ScriptStructs.Num(); ScriptIndex++)
-						{
-							ScriptStructs[ScriptIndex]->PrepareCppStructOps();
-							check(ScriptStructs[ScriptIndex]->GetCppStructOps());
-						}
-						HotReloadAr.Logf(ELogVerbosity::Warning, TEXT("HotReload successful (%d functions remapped  %d scriptstructs remapped)"), Count, ScriptStructs.Num() );
-
-						HotReloadFunctionRemap.Empty();
-					}
-				}
-				else if( bRecompileFinished )
-				{
-					HotReloadAr.Logf(ELogVerbosity::Warning, TEXT("HotReload failed, recompile failed") );
-				}
-			}
-		};
-
-
-		const double StartTime = FPlatformTime::Seconds();
-
-		TArray< FName > ModuleNames;
-		for( UPackage* Package : InPackages )
-		{
-			// Attempt to recompile this package's module
-			FName ShortPackageName =  FPackageName::GetShortFName(Package->GetFName());
-			ModuleNames.Add( ShortPackageName );
-		}
-
-		// Add dependent modules.
-		ModuleNames.Append( DependentModules );
-
-		// Start compiling modules
-		const bool bCompileStarted = FModuleManager::Get().RecompileModulesAsync(
-			ModuleNames,
-			FModuleManager::FRecompileModulesCallback::CreateStatic< TArray<UPackage*>, TArray<FName>, FOutputDevice& >( &Local::DoHotReload, InPackages, DependentModules, Ar ),
-			bWaitForCompletion,
-			Ar );
-
-		if( bCompileStarted )
-		{
-			if( bWaitForCompletion )
-			{
-				Ar.Logf(ELogVerbosity::Warning, TEXT("HotReload operation took %4.1fs."), float(FPlatformTime::Seconds() - StartTime) );
-			}
-			else
-			{
-				Ar.Logf(ELogVerbosity::Warning, TEXT("Starting HotReload took %4.1fs."), float(FPlatformTime::Seconds() - StartTime) );
-			}
-		}
-		else
-		{
-			Ar.Logf(ELogVerbosity::Warning, TEXT("RebindPackages failed because the compiler could not be started.") );
-		}
-	}
-	else
-#endif
-	{
-		Ar.Logf(ELogVerbosity::Warning, TEXT("RebindPackages not possible for specified packages (or application was compiled in monolithic mode.)") );
-	}
-}
-
-
 /*-----------------------------------------------------------------------------
    Command line.
 -----------------------------------------------------------------------------*/
diff --git a/Engine/Source/Runtime/CoreUObject/Private/UObject/ScriptCore.cpp b/Engine/Source/Runtime/CoreUObject/Private/UObject/ScriptCore.cpp
--- a/Engine/Source/Runtime/CoreUObject/Private/UObject/ScriptCore.cpp
+++ b/Engine/Source/Runtime/CoreUObject/Private/UObject/ScriptCore.cpp
@@ -6,6 +6,7 @@
 
 #include "CoreUObjectPrivate.h"
 #include "MallocProfiler.h"
+#include "HotReloadInterface.h"
 
 DEFINE_LOG_CATEGORY(LogScriptFrame);
 DEFINE_LOG_CATEGORY_STATIC(LogScriptCore, Log, All);
@@ -272,8 +273,9 @@ COREUOBJECT_API uint8 GRegisterNative( int32 NativeBytecodeIndex, const Native&
 #if !IS_MONOLITHIC
 			if (GIsHotReload)
 			{
-				CA_SUPPRESS(6385)
-				AddHotReloadFunctionRemap(Func, GNatives[NativeBytecodeIndex]);
+				IHotReloadInterface& HotReloadSupport = FModuleManager::LoadModuleChecked<IHotReloadInterface>("HotReload");
+				CA_SUPPRESS(6385)				
+				HotReloadSupport.AddHotReloadFunctionRemap(Func, GNatives[NativeBytecodeIndex]);
 			}
 			else
 #endif
diff --git a/Engine/Source/Runtime/CoreUObject/Public/Misc/HotReloadInterface.h b/Engine/Source/Runtime/CoreUObject/Public/Misc/HotReloadInterface.h
new file mode 100644
--- /dev/null
+++ b/Engine/Source/Runtime/CoreUObject/Public/Misc/HotReloadInterface.h
@@ -0,0 +1,37 @@
+// Copyright 1998-2014 Epic Games, Inc. All Rights Reserved.
+
+#pragma once
+
+#include "ModuleManager.h"
+
+/**
+* HotReload module interface
+*/
+class COREUOBJECT_API IHotReloadInterface : public IModuleInterface
+{
+public:
+
+	/**
+	* Adds a function to re-map after hot-reload.
+	*/
+	virtual void AddHotReloadFunctionRemap(Native NewFunctionPointer, Native OldFunctionPointer) = 0;
+
+	/**
+	* Performs hot reload from the editor
+	*/
+	virtual void DoHotReloadFromEditor() = 0;
+
+	/**
+	* HotReload: Reloads the DLLs for given packages.
+	* @param	Package				Packages to reload.
+	* @param	DependentModules	Additional modules that don't contain UObjects, but rely on them
+	* @param	bWaitForCompletion	True if RebindPackages should not return until the recompile and reload has completed
+	* @param	Ar					Output device for logging compilation status
+	*/
+	virtual void RebindPackages(TArray<UPackage*> Packages, TArray<FName> DependentModules, const bool bWaitForCompletion, FOutputDevice &Ar) = 0;
+
+	/** Called when a Hot Reload event has completed. */
+	DECLARE_MULTICAST_DELEGATE(FHotReloadEvent);
+	virtual FHotReloadEvent& OnHotReload() = 0;
+};
+
diff --git a/Engine/Source/Runtime/CoreUObject/Public/UObject/UObjectGlobals.h b/Engine/Source/Runtime/CoreUObject/Public/UObject/UObjectGlobals.h
--- a/Engine/Source/Runtime/CoreUObject/Public/UObject/UObjectGlobals.h
+++ b/Engine/Source/Runtime/CoreUObject/Public/UObject/UObjectGlobals.h
@@ -51,11 +51,6 @@ extern int32						GForcedExportCount;
 /** Objects that might need preloading.									*/
 extern TArray<UObject*>			GObjLoaded;
 
-#if !IS_MONOLITHIC
-	/** Adds and entry for the UFunction native pointer remap table */
-	COREUOBJECT_API void	AddHotReloadFunctionRemap(Native NewFunctionPointer, Native OldFunctionPointer);
-#endif
-
 #if !(UE_BUILD_SHIPPING || UE_BUILD_TEST)
 /** Used to verify that the Super::Serialize chain is intact.			*/
 extern TArray<UObject*,TInlineAllocator<16> >		DebugSerialize;
@@ -427,15 +422,6 @@ COREUOBJECT_API UPackage* CreatePackage( UObject* InOuter, const TCHAR* PackageN
 void StaticShutdownAfterError();
 void GlobalSetProperty( const TCHAR* Value, UClass* Class, UProperty* Property, bool bNotifyObjectOfChange );
 
-/**
- * HotReload: Reloads the DLLs for given packages. 
- * @param	Package				Packages to reload.
- * @param	DependentModules	Additional modules that don't contain UObjects, but rely on them
- * @param	bWaitForCompletion	True if RebindPackages should not return until the recompile and reload has completed
- * @param	Ar					Output device for logging compilation status
-*/
-COREUOBJECT_API void RebindPackages( TArray< UPackage* > Packages, TArray< FName > DependentModules, const bool bWaitForCompletion, FOutputDevice &Ar );
-
 /**
  * Call back into the async loading code to inform of the creation of a new object
  * @param Object		Object created
diff --git a/Engine/Source/Runtime/Engine/Classes/Engine/Engine.h b/Engine/Source/Runtime/Engine/Classes/Engine/Engine.h
--- a/Engine/Source/Runtime/Engine/Classes/Engine/Engine.h
+++ b/Engine/Source/Runtime/Engine/Classes/Engine/Engine.h
@@ -1628,7 +1628,9 @@ public:
 
 	// Compile in Debug or Development
 #if !(UE_BUILD_SHIPPING || UE_BUILD_TEST)
+#if !IS_MONOLITHIC
 	bool HandleHotReloadCommand( const TCHAR* Cmd, FOutputDevice& Ar );
+#endif
 	bool HandleDumpConsoleCommandsCommand( const TCHAR* Cmd, FOutputDevice& Ar, UWorld* InWorld );
 	bool HandleShowMaterialDrawEventsCommand( const TCHAR* Cmd, FOutputDevice& Ar );
 	bool HandleDumpAvailableResolutionsCommand( const TCHAR* Cmd, FOutputDevice& Ar );
diff --git a/Engine/Source/Runtime/Engine/Private/UnrealEngine.cpp b/Engine/Source/Runtime/Engine/Private/UnrealEngine.cpp
--- a/Engine/Source/Runtime/Engine/Private/UnrealEngine.cpp
+++ b/Engine/Source/Runtime/Engine/Private/UnrealEngine.cpp
@@ -81,6 +81,7 @@
 #include "ContentStreaming.h"
 
 #include "Engine/GameInstance.h"
+#include "HotReloadInterface.h"
 
 DEFINE_LOG_CATEGORY_STATIC(LogEngine, Log, All);
 
@@ -2206,10 +2207,12 @@ bool UEngine::Exec( UWorld* InWorld, const TCHAR* Cmd, FOutputDevice& Ar )
 #endif
 
 #if !(UE_BUILD_SHIPPING || UE_BUILD_TEST)
+#if !IS_MONOLITHIC
 	else if( FParse::Command(&Cmd,TEXT("HotReload")) )
 	{
 		return HandleHotReloadCommand( Cmd, Ar );
 	}
+#endif // !IS_MONOLITHIC
 	else if (FParse::Command(&Cmd, TEXT("DumpConsoleCommands")))
 	{
 		return HandleDumpConsoleCommandsCommand( Cmd, Ar, InWorld );
@@ -2634,6 +2637,7 @@ bool UEngine::HandleMergeMeshCommand( const TCHAR* Cmd, FOutputDevice& Ar, UWorl
 
 
 #if !(UE_BUILD_SHIPPING || UE_BUILD_TEST)
+#if !IS_MONOLITHIC
 bool UEngine::HandleHotReloadCommand( const TCHAR* Cmd, FOutputDevice& Ar )
 {
 	FString Module(FParse::Token(Cmd, 0));
@@ -2649,10 +2653,12 @@ bool UEngine::HandleHotReloadCommand( const TCHAR* Cmd, FOutputDevice& Ar )
 		TArray< UPackage*> PackagesToRebind;
 		PackagesToRebind.Add( Package );
 		const bool bWaitForCompletion = true;	// Always wait when hotreload is initiated from the console
-		RebindPackages(PackagesToRebind, TArray<FName>(), bWaitForCompletion, Ar);
+		IHotReloadInterface& HotReloadSupport = FModuleManager::LoadModuleChecked<IHotReloadInterface>("HotReload");
+		HotReloadSupport.RebindPackages(PackagesToRebind, TArray<FName>(), bWaitForCompletion, Ar);
 	}
-	return 1;
+	return true;
 }
+#endif // !IS_MONOLITHIC
 
 bool UEngine::HandleDumpConsoleCommandsCommand( const TCHAR* Cmd, FOutputDevice& Ar, UWorld* InWorld )
 {
