Fixed a crash on shutdown in HotReload module

diff --git a/Engine/Source/Developer/HotReload/Private/HotReload.cpp b/Engine/Source/Developer/HotReload/Private/HotReload.cpp
--- a/Engine/Source/Developer/HotReload/Private/HotReload.cpp
+++ b/Engine/Source/Developer/HotReload/Private/HotReload.cpp
@@ -488,12 +488,15 @@ void FHotReloadModule::InitHotReloadWatcher()
 
 void FHotReloadModule::ShutdownHotReloadWatcher()
 {
-	FDirectoryWatcherModule& DirectoryWatcherModule = FModuleManager::Get().LoadModuleChecked<FDirectoryWatcherModule>(TEXT("DirectoryWatcher"));
-	IDirectoryWatcher* DirectoryWatcher = DirectoryWatcherModule.Get();
-	if (DirectoryWatcher)
+	FDirectoryWatcherModule* DirectoryWatcherModule = FModuleManager::LoadModulePtr<FDirectoryWatcherModule>(TEXT("DirectoryWatcher"));
+	if( DirectoryWatcherModule != nullptr )
 	{
-		FString BinariesPath = FPaths::ConvertRelativePathToFull(FPaths::GameDir() / TEXT("Binaries") / FPlatformProcess::GetBinariesSubdirectory());
-		DirectoryWatcher->UnregisterDirectoryChangedCallback(BinariesPath, BinariesFolderChangedDelegate);
+		IDirectoryWatcher* DirectoryWatcher = DirectoryWatcherModule->Get();
+		if (DirectoryWatcher)
+		{
+			FString BinariesPath = FPaths::ConvertRelativePathToFull(FPaths::GameDir() / TEXT("Binaries") / FPlatformProcess::GetBinariesSubdirectory());
+			DirectoryWatcher->UnregisterDirectoryChangedCallback(BinariesPath, BinariesFolderChangedDelegate);
+		}
 	}
 }
 
