Fixed crash on shutdown

diff --git a/Engine/Source/Developer/HotReload/Private/HotReload.cpp b/Engine/Source/Developer/HotReload/Private/HotReload.cpp
--- a/Engine/Source/Developer/HotReload/Private/HotReload.cpp
+++ b/Engine/Source/Developer/HotReload/Private/HotReload.cpp
@@ -488,7 +488,7 @@ void FHotReloadModule::InitHotReloadWatcher()
 
 void FHotReloadModule::ShutdownHotReloadWatcher()
 {
-	FDirectoryWatcherModule* DirectoryWatcherModule = FModuleManager::LoadModulePtr<FDirectoryWatcherModule>(TEXT("DirectoryWatcher"));
+	FDirectoryWatcherModule* DirectoryWatcherModule = FModuleManager::GetModulePtr<FDirectoryWatcherModule>(TEXT("DirectoryWatcher"));
 	if( DirectoryWatcherModule != nullptr )
 	{
 		IDirectoryWatcher* DirectoryWatcher = DirectoryWatcherModule->Get();
