Move hot reload functionality from Core into the HotReload module. Precursor to removing assumptions about IS_MONOLITHIC and supporting hybrid monolithic configurations.

diff --git a/Engine/Plugins/Developer/VisualStudioSourceCodeAccess/Source/VisualStudioSourceCodeAccess/Private/VisualStudioSourceCodeAccessor.cpp b/Engine/Plugins/Developer/VisualStudioSourceCodeAccess/Source/VisualStudioSourceCodeAccess/Private/VisualStudioSourceCodeAccessor.cpp
--- a/Engine/Plugins/Developer/VisualStudioSourceCodeAccess/Source/VisualStudioSourceCodeAccess/Private/VisualStudioSourceCodeAccessor.cpp
+++ b/Engine/Plugins/Developer/VisualStudioSourceCodeAccess/Source/VisualStudioSourceCodeAccess/Private/VisualStudioSourceCodeAccessor.cpp
@@ -6,6 +6,7 @@
 #include "ISourceCodeAccessModule.h"
 #include "ModuleManager.h"
 #include "DesktopPlatformModule.h"
+#include "Developer/HotReload/Public/IHotReload.h"
 
 
 #include "AllowWindowsPlatformTypes.h"
@@ -67,7 +68,7 @@ void FVisualStudioSourceCodeAccessor::Startup()
 	VSLaunchTime = 0.0;
 
 	// Setup compilation for saving all VS documents upon compilation start
-	FModuleManager::Get().OnModuleCompilerStarted().AddStatic( &SaveVisualStudioDocuments );
+	IHotReloadModule::Get().OnModuleCompilerStarted().AddStatic( &SaveVisualStudioDocuments );
 
 	// Cache this so we don't have to do it on a background thread
 	GetSolutionPath();
@@ -82,7 +83,8 @@ void FVisualStudioSourceCodeAccessor::Startup()
 
 void FVisualStudioSourceCodeAccessor::Shutdown()
 {
-	FModuleManager::Get().OnModuleCompilerStarted().RemoveStatic( &SaveVisualStudioDocuments );
+	// Unregister the hot-reload callback
+	IHotReloadModule::Get().OnModuleCompilerStarted().RemoveStatic( &SaveVisualStudioDocuments );
 }
 
 #if VSACCESSOR_HAS_DTE
diff --git a/Engine/Plugins/Developer/VisualStudioSourceCodeAccess/Source/VisualStudioSourceCodeAccess/VisualStudioSourceCodeAccess.Build.cs b/Engine/Plugins/Developer/VisualStudioSourceCodeAccess/Source/VisualStudioSourceCodeAccess/VisualStudioSourceCodeAccess.Build.cs
--- a/Engine/Plugins/Developer/VisualStudioSourceCodeAccess/Source/VisualStudioSourceCodeAccess/VisualStudioSourceCodeAccess.Build.cs
+++ b/Engine/Plugins/Developer/VisualStudioSourceCodeAccess/Source/VisualStudioSourceCodeAccess/VisualStudioSourceCodeAccess.Build.cs
@@ -11,7 +11,8 @@ namespace UnrealBuildTool.Rules
 				{
 					"Core",
 					"SourceCodeAccess",
-					"DesktopPlatform"
+					"DesktopPlatform",
+					"HotReload"
 				}
 			);
 
diff --git a/Engine/Source/Developer/HotReload/Private/HotReload.cpp b/Engine/Source/Developer/HotReload/Private/HotReload.cpp
--- a/Engine/Source/Developer/HotReload/Private/HotReload.cpp
+++ b/Engine/Source/Developer/HotReload/Private/HotReload.cpp
@@ -13,22 +13,82 @@ DEFINE_LOG_CATEGORY(LogHotReload);
 /**
  * Module for HotReload support
  */
-class FHotReloadModule : public IHotReloadModule
+class FHotReloadModule : public IHotReloadModule, FSelfRegisteringExec
 {
 public:
+	FHotReloadModule()
+	{
+		ModuleCompileReadPipe = nullptr;
+		bRequestCancelCompilation = false;
+	}
 
 	/** IModuleInterface implementation */
 	virtual void StartupModule() override;
 	virtual void ShutdownModule() override;
 
+	/** FSelfRegisteringExec implementation */
+	virtual bool Exec( UWorld* Inworld, const TCHAR* Cmd, FOutputDevice& Ar ) override;
+
 	/** IHotReloadInterface implementation */
+	virtual void Tick() override;
+	virtual void SaveConfig() override;
+	virtual bool RecompileModule( const FName InModuleName, const bool bReloadAfterRecompile, FOutputDevice &Ar ) override;
+	virtual bool IsCurrentlyCompiling() const override { return ModuleCompileProcessHandle.IsValid(); }
+	virtual void RequestStopCompilation() override { bRequestCancelCompilation = true; }
 	virtual void AddHotReloadFunctionRemap(Native NewFunctionPointer, Native OldFunctionPointer) override;	
 	virtual void RebindPackages(TArray< UPackage* > Packages, TArray< FName > DependentModules, const bool bWaitForCompletion, FOutputDevice &Ar) override;
 	virtual void DoHotReloadFromEditor() override;
-	virtual FHotReloadEvent& OnHotReload() { return HotReloadEvent; }	
+	virtual FHotReloadEvent& OnHotReload() override { return HotReloadEvent; }	
+	virtual FModuleCompilerStartedEvent& OnModuleCompilerStarted() override { return ModuleCompilerStartedEvent; }
+	virtual FModuleCompilerFinishedEvent& OnModuleCompilerFinished() override { return ModuleCompilerFinishedEvent; }
+	virtual FString GetModuleCompileMethod(FName InModuleName) override;
 
 private:
-	
+	/**
+	 * Enumerates compilation methods for modules.
+	 */
+	enum class EModuleCompileMethod
+	{
+		Runtime,
+		External,
+		Unknown
+	};
+
+	/**
+	 * Helper structure to hold on to module state while asynchronously recompiling DLLs
+	 */
+	struct FModuleToRecompile
+	{
+		/** Name of the module */
+		FString ModuleName;
+
+		/** Desired module file name suffix, or empty string if not needed */
+		FString ModuleFileSuffix;
+
+		/** The module file name to use after a compilation succeeds, or an empty string if not changing */
+		FString NewModuleFilename;
+	};
+
+	/**
+	 * Helper structure to store the compile time and method for a module
+	 */
+	struct FModuleCompilationData
+	{
+		/** Has a timestamp been set for the .dll file */
+		bool bHasFileTimeStamp;
+
+		/** Last known timestamp for the .dll file */
+		FDateTime FileTimeStamp;
+
+		/** Last known compilation method of the .dll file */
+		EModuleCompileMethod CompileMethod;
+
+		FModuleCompilationData()
+			: bHasFileTimeStamp(false)
+			, CompileMethod(EModuleCompileMethod::Unknown)
+		{ }
+	};
+
 	/**
 	 * Adds a callback to directory watcher for the game binaries folder.
 	 */
@@ -96,9 +156,78 @@ private:
 
 	/**
 	 * Sends analytics event about the re-load
-	*/
+	 */
 	static void RecordAnalyticsEvent(const TCHAR* ReloadFrom, ECompilationResult::Type Result, double Duration, int32 PackageCount, int32 DependentModulesCount);
 
+	/**
+	 * Declares a type of delegates that is executed after a module recompile has finished.
+	 *
+	 * The first argument signals whether compilation has finished.
+	 * The second argument shows whether compilation was successful or not.
+	 */
+	DECLARE_DELEGATE_TwoParams( FRecompileModulesCallback, bool, bool );
+
+	/**
+	 * Tries to recompile the specified modules in the background.  When recompiling finishes, the specified callback
+	 * delegate will be triggered, passing along a bool that tells you whether the compile action succeeded.  This
+	 * function never tries to unload modules or to reload the modules after they finish compiling.  You should do
+	 * that yourself in the recompile completion callback!
+	 *
+	 * @param ModuleNames Names of the modules to recompile
+	 * @param RecompileModulesCallback Callback function to execute after compilation finishes (whether successful or not.)
+	 * @param bWaitForCompletion True if the function should not return until recompilation attempt has finished and callbacks have fired
+	 * @param Ar Output device for logging compilation status
+	 * @return	True if the recompile action was kicked off successfully.  If this returns false, then the recompile callback will never fire.  In the case where bWaitForCompletion=false, this will also return false if the compilation failed for any reason.
+	 */
+	bool RecompileModulesAsync( const TArray< FName > ModuleNames, const FRecompileModulesCallback& InRecompileModulesCallback, const bool bWaitForCompletion, FOutputDevice &Ar );
+
+	/** Called for successfully re-complied module */
+	void OnModuleCompileSucceeded(FName ModuleName, const FString& NewModuleFilename);
+
+	/**
+	 * Tries to recompile the specified DLL using UBT. Does not interact with modules. This is a low level routine.
+	 *
+	 * @param ModuleNames List of modules to recompile, including the module name and optional file suffix.
+	 * @param Ar Output device for logging compilation status.
+	 */
+	bool RecompileModuleDLLs( const TArray< FModuleToRecompile >& ModuleNames, FOutputDevice& Ar );
+
+	/** Returns arguments to pass to UnrealBuildTool when compiling modules */
+	static FString MakeUBTArgumentsForModuleCompiling();
+
+	/** 
+	 *	Starts compiling DLL files for one or more modules.
+	 *
+	 *	@param GameName The name of the game.
+	 *	@param ModuleNames The list of modules to compile.
+	 *	@param InRecompileModulesCallback Callback function to make when module recompiles.
+	 *	@param Ar
+	 *	@param bInFailIfGeneratedCodeChanges If true, fail the compilation if generated headers change.
+	 *	@param InAdditionalCmdLineArgs Additional arguments to pass to UBT.
+	 *	@return true if successful, false otherwise.
+	 */
+	bool StartCompilingModuleDLLs(const FString& GameName, const TArray< FModuleToRecompile >& ModuleNames, 
+		const FRecompileModulesCallback& InRecompileModulesCallback, FOutputDevice& Ar, bool bInFailIfGeneratedCodeChanges, 
+		const FString& InAdditionalCmdLineArgs = FString() );
+
+	/** Launches UnrealBuildTool with the specified command line parameters */
+	bool InvokeUnrealBuildToolForCompile(const FString& InCmdLineParams, FOutputDevice &Ar);
+
+	/** Checks to see if a pending compilation action has completed and optionally waits for it to finish.  If completed, fires any appropriate callbacks and reports status provided bFireEvents is true. */
+	void CheckForFinishedModuleDLLCompile(const bool bWaitForCompletion, bool& bCompileStillInProgress, bool& bCompileSucceeded, FOutputDevice& Ar, const FText& SlowTaskOverrideText = FText::GetEmpty(), bool bFireEvents = true);
+
+	/** Called when the compile data for a module need to be update in memory and written to config */
+	void UpdateModuleCompileData(FName ModuleName);
+
+	/** Called when a new module is added to the manager to get the saved compile data from config */
+	static void ReadModuleCompilationInfoFromConfig(FName ModuleName, FModuleCompilationData& CompileData);
+
+	/** Saves the module's compile data to config */
+	static void WriteModuleCompilationInfoToConfig(FName ModuleName, const FModuleCompilationData& CompileData);
+
+	/** Access the module's file and read the timestamp from the file system. Returns true if the timestamp was read successfully. */
+	bool GetModuleFileTimeStamp(FName ModuleName, FDateTime& OutFileTimeStamp) const;
+
 	/** FTicker delegate (hot-reload from IDE) */
 	FTickerDelegate TickerDelegate;
 
@@ -124,8 +253,52 @@ private:
 
 	/** Delegate broadcast when a module has been hot-reloaded */
 	FHotReloadEvent HotReloadEvent;
+
+	/** Array of modules that we're currently recompiling */
+	TArray< FModuleToRecompile > ModulesBeingCompiled;
+
+	/** Array of modules that we're going to recompile */
+	TArray< FModuleToRecompile > ModulesThatWereBeingRecompiled;
+
+	/** Last known compilation data for each module */
+	TMap<FName, TSharedRef<FModuleCompilationData>> ModuleCompileData;
+
+	/** Multicast delegate which will broadcast a notification when the compiler starts */
+	FModuleCompilerStartedEvent ModuleCompilerStartedEvent;
+	
+	/** Multicast delegate which will broadcast a notification when the compiler finishes */
+	FModuleCompilerFinishedEvent ModuleCompilerFinishedEvent;
+
+	/** When compiling a module using an external application, stores the handle to the process that is running */
+	FProcHandle ModuleCompileProcessHandle;
+
+	/** When compiling a module using an external application, this is the process read pipe handle */
+	void* ModuleCompileReadPipe;
+
+	/** When compiling a module using an external application, this is the text that was read from the read pipe handle */
+	FString ModuleCompileReadPipeText;
+
+	/** Callback to execute after an asynchronous recompile has completed (whether successful or not.) */
+	FRecompileModulesCallback RecompileModulesCallback;
+
+	/** true if we should attempt to cancel the current async compilation */
+	bool bRequestCancelCompilation;
 };
 
+namespace HotReloadDefs
+{
+	const static FString CompilationInfoConfigSection("ModuleFileTracking");
+
+	// These strings should match the values of the enum EModuleCompileMethod in ModuleManager.h
+	// and should be handled in ReadModuleCompilationInfoFromConfig() & WriteModuleCompilationInfoToConfig() below
+	const static FString CompileMethodRuntime("Runtime");
+	const static FString CompileMethodExternal("External");
+	const static FString CompileMethodUnknown("Unknown");
+
+	// Add one minute epsilon to timestamp comparision
+	const static FTimespan TimeStampEpsilon(0, 1, 0);
+}
+
 IMPLEMENT_MODULE(FHotReloadModule, HotReload);
 
 void FHotReloadModule::StartupModule()
@@ -149,6 +322,156 @@ void FHotReloadModule::ShutdownModule()
 	ShutdownHotReloadWatcher();
 }
 
+bool FHotReloadModule::Exec( UWorld* Inworld, const TCHAR* Cmd, FOutputDevice& Ar )
+{
+#if !UE_BUILD_SHIPPING
+	if ( FParse::Command( &Cmd, TEXT( "Module" ) ) )
+	{
+#if !IS_MONOLITHIC
+		// Recompile <ModuleName>
+		if( FParse::Command( &Cmd, TEXT( "Recompile" ) ) )
+		{
+			const FString ModuleNameStr = FParse::Token( Cmd, 0 );
+			if( !ModuleNameStr.IsEmpty() )
+			{
+				const FName ModuleName( *ModuleNameStr );
+				const bool bReloadAfterRecompile = true;
+				RecompileModule( ModuleName, bReloadAfterRecompile, Ar);
+			}
+
+			return true;
+		}
+#endif // !IS_MONOLITHIC
+	}
+#endif // !UE_BUILD_SHIPPING
+	return false;
+}
+
+void FHotReloadModule::Tick()
+{
+	// We never want to block on a pending compile when checking compilation status during Tick().  We're
+	// just checking so that we can fire callbacks if and when compilation has finished.
+	const bool bWaitForCompletion = false;
+
+	// Ignored output variables
+	bool bCompileStillInProgress = false;
+	bool bCompileSucceeded = false;
+	FOutputDeviceNull NullOutput;
+	CheckForFinishedModuleDLLCompile( bWaitForCompletion, bCompileStillInProgress, bCompileSucceeded, NullOutput );
+}
+	
+void FHotReloadModule::SaveConfig()
+{
+	// Find all the modules
+	TArray<FModuleStatus> Modules;
+	FModuleManager::Get().QueryModules(Modules);
+
+	// Update the compile data for each one
+	for( const FModuleStatus &Module : Modules )
+	{
+		UpdateModuleCompileData(*Module.Name);
+	}
+}
+
+FString FHotReloadModule::GetModuleCompileMethod(FName InModuleName)
+{
+	if (!ModuleCompileData.Contains(InModuleName))
+	{
+		UpdateModuleCompileData(InModuleName);
+	}
+
+	switch(ModuleCompileData.FindChecked(InModuleName).Get().CompileMethod)
+	{
+	case EModuleCompileMethod::External:
+		return HotReloadDefs::CompileMethodExternal;
+	case EModuleCompileMethod::Runtime:
+		return HotReloadDefs::CompileMethodRuntime;
+	default:
+		return HotReloadDefs::CompileMethodUnknown;
+	}
+}
+
+bool FHotReloadModule::RecompileModule( const FName InModuleName, const bool bReloadAfterRecompile, FOutputDevice &Ar )
+{
+#if !IS_MONOLITHIC
+	const bool bShowProgressDialog = true;
+	const bool bShowCancelButton = false;
+
+	FFormatNamedArguments Args;
+	Args.Add( TEXT("CodeModuleName"), FText::FromName( InModuleName ) );
+	const FText StatusUpdate = FText::Format( NSLOCTEXT("ModuleManager", "Recompile_SlowTaskName", "Compiling {CodeModuleName}..."), Args );
+
+	GWarn->BeginSlowTask( StatusUpdate, bShowProgressDialog, bShowCancelButton );
+
+	ModuleCompilerStartedEvent.Broadcast();
+
+	// Update our set of known modules, in case we don't already know about this module
+	FModuleManager::Get().AddModule( InModuleName );
+
+	// Only use rolling module names if the module was already loaded into memory.  This allows us to try compiling
+	// the module without actually having to unload it first.
+	const bool bWasModuleLoaded = FModuleManager::Get().IsModuleLoaded( InModuleName );
+	const bool bUseRollingModuleNames = bWasModuleLoaded;
+
+	bool bWasSuccessful = true;
+	FString NewModuleFileNameOnSuccess = FModuleManager::Get().GetModuleFilename(InModuleName);
+	if( bUseRollingModuleNames )
+	{
+		// First, try to compile the module.  If the module is already loaded, we won't unload it quite yet.  Instead
+		// make sure that it compiles successfully.
+
+		// Find a unique file name for the module
+		FString UniqueSuffix;
+		FString UniqueModuleFileName;
+		FModuleManager::Get().MakeUniqueModuleFilename( InModuleName, UniqueSuffix, UniqueModuleFileName );
+
+		// If the recompile succeeds, we'll update our cached file name to use the new unique file name
+		// that we setup for the module
+		NewModuleFileNameOnSuccess = UniqueModuleFileName;
+
+		TArray< FModuleToRecompile > ModulesToRecompile;
+		FModuleToRecompile ModuleToRecompile;
+		ModuleToRecompile.ModuleName = InModuleName.ToString();
+		ModuleToRecompile.ModuleFileSuffix = UniqueSuffix;
+		ModuleToRecompile.NewModuleFilename = UniqueModuleFileName;
+		ModulesToRecompile.Add( ModuleToRecompile );
+		bWasSuccessful = RecompileModuleDLLs( ModulesToRecompile, Ar );
+	}
+
+	if( bWasSuccessful )
+	{
+		// Shutdown the module if it's already running
+		if( bWasModuleLoaded )
+		{
+			Ar.Logf( TEXT( "Unloading module before compile." ) );
+			FModuleManager::Get().UnloadOrAbandonModuleWithCallback( InModuleName, Ar );
+		}
+
+		if( !bUseRollingModuleNames )
+		{
+			// Try to recompile the DLL
+			TArray< FModuleToRecompile > ModulesToRecompile;
+			FModuleToRecompile ModuleToRecompile;
+			ModuleToRecompile.ModuleName = InModuleName.ToString();
+			ModulesToRecompile.Add( ModuleToRecompile );
+			bWasSuccessful = RecompileModuleDLLs( ModulesToRecompile, Ar );
+		}
+
+		// Reload the module if it was loaded before we recompiled
+		if( bWasSuccessful && bWasModuleLoaded && bReloadAfterRecompile )
+		{
+			Ar.Logf( TEXT( "Reloading module after successful compile." ) );
+			bWasSuccessful = FModuleManager::Get().LoadModuleWithCallback( InModuleName, Ar );
+		}
+	}
+
+	GWarn->EndSlowTask();
+	return bWasSuccessful;
+#else
+	return false;
+#endif // !IS_MONOLITHIC
+}
+
 /** Type hash for a UObject Function Pointer, maybe not a great choice, but it should be sufficient for the needs here. **/
 inline uint32 GetTypeHash(Native A)
 {
@@ -356,7 +679,7 @@ ECompilationResult::Type FHotReloadModule::RebindPackagesInternal(TArray<UPackag
 	}
 
 	// We can only proceed if a compile isn't already in progress
-	if (FModuleManager::Get().IsCurrentlyCompiling())
+	if (IsCurrentlyCompiling())
 	{
 		Ar.Logf(ELogVerbosity::Warning, TEXT("Could not rebind package because a module compile is already in progress."));
 		bCanRebind = false;
@@ -380,9 +703,9 @@ ECompilationResult::Type FHotReloadModule::RebindPackagesInternal(TArray<UPackag
 		ModuleNames.Append(DependentModules);
 
 		// Start compiling modules
-		const bool bCompileStarted = FModuleManager::Get().RecompileModulesAsync(
+		const bool bCompileStarted = RecompileModulesAsync(
 			ModuleNames,
-			FModuleManager::FRecompileModulesCallback::CreateRaw< FHotReloadModule, TArray<UPackage*>, TArray<FName>, FOutputDevice& >(this, &FHotReloadModule::DoHotReloadCallback, InPackages, DependentModules, Ar),
+			FRecompileModulesCallback::CreateRaw< FHotReloadModule, TArray<UPackage*>, TArray<FName>, FOutputDevice& >(this, &FHotReloadModule::DoHotReloadCallback, InPackages, DependentModules, Ar),
 			bWaitForCompletion,
 			Ar);
 
@@ -556,8 +879,8 @@ void FHotReloadModule::DoHotReloadFromIDE()
 		// Update compile data before we start compiling
 		for (auto& NewModule : NewModules)
 		{
-			FModuleManager::Get().UpdateModuleCompileData(*NewModule.Name);
-			FModuleManager::Get().OnModuleCompileSucceeded(*NewModule.Name, NewModule.NewFilename);
+			UpdateModuleCompileData(*NewModule.Name);
+			OnModuleCompileSucceeded(*NewModule.Name, NewModule.NewFilename);
 		}
 
 
@@ -590,4 +913,514 @@ void FHotReloadModule::RecordAnalyticsEvent(const TCHAR* ReloadFrom, ECompilatio
 	}
 }
 
+bool FHotReloadModule::RecompileModulesAsync( const TArray< FName > ModuleNames, const FRecompileModulesCallback& InRecompileModulesCallback, const bool bWaitForCompletion, FOutputDevice &Ar )
+{
+#if !IS_MONOLITHIC
+	// NOTE: This method of recompiling always using a rolling file name scheme, since we never want to unload before
+	// we start recompiling, and we need the output DLL to be unlocked before we invoke the compiler
+
+	ModuleCompilerStartedEvent.Broadcast();
+
+	TArray< FModuleToRecompile > ModulesToRecompile;
+
+	for( TArray< FName >::TConstIterator CurModuleIt( ModuleNames ); CurModuleIt; ++CurModuleIt )
+	{
+		const FName CurModuleName = *CurModuleIt;
+
+		// Update our set of known modules, in case we don't already know about this module
+		FModuleManager::Get().AddModule( CurModuleName );
+
+		FString NewModuleFileNameOnSuccess = FModuleManager::Get().GetModuleFilename(CurModuleName);
+
+		// Find a unique file name for the module
+		FString UniqueSuffix;
+		FString UniqueModuleFileName;
+		FModuleManager::Get().MakeUniqueModuleFilename( CurModuleName, UniqueSuffix, UniqueModuleFileName );
+
+		// If the recompile succeeds, we'll update our cached file name to use the new unique file name
+		// that we setup for the module
+		NewModuleFileNameOnSuccess = UniqueModuleFileName;
+
+		FModuleToRecompile ModuleToRecompile;
+		ModuleToRecompile.ModuleName = CurModuleName.ToString();
+		ModuleToRecompile.ModuleFileSuffix = UniqueSuffix;
+		ModuleToRecompile.NewModuleFilename = UniqueModuleFileName;
+		ModulesToRecompile.Add( ModuleToRecompile );
+	}
+
+	// Kick off compilation!
+	const FString AdditionalArguments = MakeUBTArgumentsForModuleCompiling();
+	const bool bFailIfGeneratedCodeChanges = false;
+	bool bWasSuccessful = StartCompilingModuleDLLs( FApp::GetGameName(), ModulesToRecompile, InRecompileModulesCallback, Ar, bFailIfGeneratedCodeChanges, AdditionalArguments );
+	if (bWasSuccessful)
+	{
+		// Go ahead and check for completion right away.  This is really just so that we can handle the case
+		// where the user asked us to wait for the compile to finish before returning.
+		bool bCompileStillInProgress = false;
+		bool bCompileSucceeded = false;
+		FOutputDeviceNull NullOutput;
+		CheckForFinishedModuleDLLCompile( bWaitForCompletion, bCompileStillInProgress, bCompileSucceeded, NullOutput );
+		if( !bCompileStillInProgress && !bCompileSucceeded )
+		{
+			bWasSuccessful = false;
+		}
+	}
+
+	return bWasSuccessful;
+#else
+	return false;
+#endif // !IS_MONOLITHIC
+}
+
+void FHotReloadModule::OnModuleCompileSucceeded(FName ModuleName, const FString& NewModuleFilename)
+{
+	// If the compile succeeded, update the module info entry with the new file name for this module
+	FModuleManager::Get().SetModuleFilename(ModuleName, NewModuleFilename);
+
+#if !IS_MONOLITHIC && WITH_EDITOR
+	// UpdateModuleCompileData() should have been run before compiling so the
+	// data in ModuleInfo should be correct for the pre-compile dll file.
+	FModuleCompilationData& CompileData = ModuleCompileData.FindChecked(ModuleName).Get();
+
+	FDateTime FileTimeStamp;
+	bool bGotFileTimeStamp = GetModuleFileTimeStamp(ModuleName, FileTimeStamp);
+
+	CompileData.bHasFileTimeStamp = bGotFileTimeStamp;
+	CompileData.FileTimeStamp = FileTimeStamp;
+
+	if (CompileData.bHasFileTimeStamp)
+	{
+		CompileData.CompileMethod = EModuleCompileMethod::Runtime;
+	}
+	else
+	{
+		CompileData.CompileMethod = EModuleCompileMethod::Unknown;
+	}
+	WriteModuleCompilationInfoToConfig(ModuleName, CompileData);
+#endif
+}
+
+bool FHotReloadModule::RecompileModuleDLLs( const TArray< FModuleToRecompile >& ModuleNames, FOutputDevice& Ar )
+{
+	bool bCompileSucceeded = false;
+#if !IS_MONOLITHIC
+	const FString AdditionalArguments = MakeUBTArgumentsForModuleCompiling();
+	if( StartCompilingModuleDLLs( FApp::GetGameName(), ModuleNames, FRecompileModulesCallback(), Ar, true, AdditionalArguments ) )
+	{
+		const bool bWaitForCompletion = true;	// Always wait
+		bool bCompileStillInProgress = false;
+		CheckForFinishedModuleDLLCompile( bWaitForCompletion, bCompileStillInProgress, bCompileSucceeded, Ar );
+	}
+#endif
+	return bCompileSucceeded;
+}
+
+FString FHotReloadModule::MakeUBTArgumentsForModuleCompiling()
+{
+	FString AdditionalArguments;
+	if ( FPaths::IsProjectFilePathSet() )
+	{
+		// We have to pass FULL paths to UBT
+		FString FullProjectPath = FPaths::ConvertRelativePathToFull(FPaths::GetProjectFilePath());
+
+		// @todo projectdirs: Currently non-Rocket projects that exist under the UE4 root are compiled by UBT with no .uproject file
+		//     name passed in (see bIsProjectTarget in VCProject.cs), which causes intermediate libraries to be saved to the Engine
+		//     intermediate folder instead of the project's intermediate folder.  We're emulating this behavior here for module
+		//     recompiling, so that compiled modules will be able to find their import libraries in the original folder they were compiled.
+		if( FRocketSupport::IsRocket() || !FullProjectPath.StartsWith( FPaths::ConvertRelativePathToFull( FPaths::RootDir() ) ) )
+		{
+			const FString ProjectFilenameWithQuotes = FString::Printf(TEXT("\"%s\""), *FullProjectPath);
+			AdditionalArguments += FString::Printf(TEXT("%s "), *ProjectFilenameWithQuotes);
+		}
+
+		if (FRocketSupport::IsRocket())
+		{
+			AdditionalArguments += TEXT("-rocket ");
+		}
+	}
+
+	return AdditionalArguments;
+}
+
+bool FHotReloadModule::StartCompilingModuleDLLs(const FString& GameName, const TArray< FModuleToRecompile >& ModuleNames, 
+	const FRecompileModulesCallback& InRecompileModulesCallback, FOutputDevice& Ar, bool bInFailIfGeneratedCodeChanges, 
+	const FString& InAdditionalCmdLineArgs )
+{
+#if PLATFORM_DESKTOP && !IS_MONOLITHIC
+	// Keep track of what we're compiling
+	ModulesBeingCompiled = ModuleNames;
+	ModulesThatWereBeingRecompiled = ModulesBeingCompiled;
+
+	const TCHAR* BuildPlatformName = FPlatformMisc::GetUBTPlatform();
+	const TCHAR* BuildConfigurationName = FModuleManager::GetUBTConfiguration();
+
+	RecompileModulesCallback = InRecompileModulesCallback;
+
+	// Pass a module file suffix to UBT if we have one
+	FString ModuleArg;
+	for( int32 CurModuleIndex = 0; CurModuleIndex < ModuleNames.Num(); ++CurModuleIndex )
+	{
+		if( !ModuleNames[ CurModuleIndex ].ModuleFileSuffix.IsEmpty() )
+		{
+			ModuleArg += FString::Printf( TEXT( " -ModuleWithSuffix %s %s" ), *ModuleNames[ CurModuleIndex ].ModuleName, *ModuleNames[ CurModuleIndex ].ModuleFileSuffix );
+		}
+		else
+		{
+			ModuleArg += FString::Printf( TEXT( " -Module %s" ), *ModuleNames[ CurModuleIndex ].ModuleName );
+		}
+		Ar.Logf( TEXT( "Recompiling %s..." ), *ModuleNames[ CurModuleIndex ].ModuleName );
+
+		// prepare the compile info in the FModuleInfo so that it can be compared after compiling
+		FName ModuleFName(*ModuleNames[ CurModuleIndex ].ModuleName);
+		UpdateModuleCompileData(ModuleFName);
+	}
+
+	FString ExtraArg;
+#if UE_EDITOR
+	// NOTE: When recompiling from the editor, we're passed the game target name, not the editor target name, but we'll
+	//       pass "-editorrecompile" to UBT which tells UBT to figure out the editor target to use for this game, since
+	//       we can't possibly know what the target is called from within the engine code.
+	ExtraArg = TEXT( "-editorrecompile " );
+#endif
+
+	if (bInFailIfGeneratedCodeChanges)
+	{
+		// Additional argument to let UHT know that we can only compile the module if the generated code didn't change
+		ExtraArg += TEXT( "-FailIfGeneratedCodeChanges " );
+	}
+
+	FString CmdLineParams = FString::Printf( TEXT( "%s%s %s %s %s%s" ), 
+		*GameName, *ModuleArg, 
+		BuildPlatformName, BuildConfigurationName, 
+		*ExtraArg, *InAdditionalCmdLineArgs );
+
+	const bool bInvocationSuccessful = InvokeUnrealBuildToolForCompile(CmdLineParams, Ar);
+	if ( !bInvocationSuccessful )
+	{
+		// No longer compiling modules
+		ModulesBeingCompiled.Empty();
+
+		ModuleCompilerFinishedEvent.Broadcast(FString(), ECompilationResult::OtherCompilationError, false);
+
+		// Fire task completion delegate 
+		
+		RecompileModulesCallback.ExecuteIfBound( false, false );
+		RecompileModulesCallback.Unbind();
+	}
+
+	return bInvocationSuccessful;
+#else
+	return false;
+#endif
+}
+
+bool FHotReloadModule::InvokeUnrealBuildToolForCompile(const FString& InCmdLineParams, FOutputDevice &Ar)
+{
+#if PLATFORM_DESKTOP && !IS_MONOLITHIC
+
+	// Make sure we're not already compiling something!
+	check(!IsCurrentlyCompiling());
+
+	// Setup output redirection pipes, so that we can harvest compiler output and display it ourselves
+#if PLATFORM_LINUX
+	int pipefd[2];
+	pipe(pipefd);
+	void* PipeRead = &pipefd[0];
+	void* PipeWrite = &pipefd[1];
+#else
+	void* PipeRead = NULL;
+	void* PipeWrite = NULL;
+#endif
+
+	verify(FPlatformProcess::CreatePipe(PipeRead, PipeWrite));
+	ModuleCompileReadPipeText = TEXT("");
+
+	FProcHandle ProcHandle = FUBTInvoker::InvokeUnrealBuildToolAsync(InCmdLineParams, Ar, PipeRead, PipeWrite);
+
+	// We no longer need the Write pipe so close it.
+	// We DO need the Read pipe however...
+#if PLATFORM_LINUX
+	close(*(int*)PipeWrite);
+#else
+	FPlatformProcess::ClosePipe(0, PipeWrite);
+#endif
+
+	if (!ProcHandle.IsValid())
+	{
+		// We're done with the process handle now
+		ModuleCompileProcessHandle.Reset();
+		ModuleCompileReadPipe = NULL;
+	}
+	else
+	{
+		ModuleCompileProcessHandle = ProcHandle;
+		ModuleCompileReadPipe = PipeRead;
+	}
+
+	return ProcHandle.IsValid();
+#else
+	return false;
+#endif // PLATFORM_DESKTOP && !IS_MONOLITHIC
+}
+
+void FHotReloadModule::CheckForFinishedModuleDLLCompile(const bool bWaitForCompletion, bool& bCompileStillInProgress, bool& bCompileSucceeded, FOutputDevice& Ar, const FText& SlowTaskOverrideText, bool bFireEvents)
+{
+#if PLATFORM_DESKTOP && !IS_MONOLITHIC
+	bCompileStillInProgress = false;
+	ECompilationResult::Type CompilationResult = ECompilationResult::OtherCompilationError;
+
+	// Is there a compilation in progress?
+	if( IsCurrentlyCompiling() )
+	{
+		bCompileStillInProgress = true;
+
+		// Ensure slow task messages are seen.
+		GWarn->PushStatus();
+
+		// Update the slow task dialog if we were summoned from a synchronous recompile path
+		if (GIsSlowTask)
+		{
+			if ( !SlowTaskOverrideText.IsEmpty() )
+			{
+				GWarn->StatusUpdate(-1, -1, SlowTaskOverrideText);
+			}
+			else
+			{
+				FText StatusUpdate;
+				if ( ModulesBeingCompiled.Num() > 0 )
+				{
+					FFormatNamedArguments Args;
+					Args.Add( TEXT("CodeModuleName"), FText::FromString( ModulesBeingCompiled[0].ModuleName ) );
+					StatusUpdate = FText::Format( NSLOCTEXT("FModuleManager", "CompileSpecificModuleStatusMessage", "{CodeModuleName}: Compiling modules..."), Args );
+				}
+				else
+				{
+					StatusUpdate = NSLOCTEXT("FModuleManager", "CompileStatusMessage", "Compiling modules...");
+				}
+
+				GWarn->StatusUpdate(-1, -1, StatusUpdate);
+			}
+		}
+
+		// Check to see if the compile has finished yet
+		int32 ReturnCode = -1;
+		while (bCompileStillInProgress)
+		{
+			if( FPlatformProcess::GetProcReturnCode( ModuleCompileProcessHandle, &ReturnCode ) )
+			{
+				bCompileStillInProgress = false;
+			}
+			
+			if (bRequestCancelCompilation)
+			{
+				FPlatformProcess::TerminateProc(ModuleCompileProcessHandle);
+				bCompileStillInProgress = bRequestCancelCompilation = false;
+			}
+
+			if( bCompileStillInProgress )
+			{
+				ModuleCompileReadPipeText += FPlatformProcess::ReadPipe(ModuleCompileReadPipe);
+
+				if( !bWaitForCompletion )
+				{
+					// We haven't finished compiling, but we were asked to return immediately
+
+					break;
+				}
+
+				// Give up a small timeslice if we haven't finished recompiling yet
+				FPlatformProcess::Sleep( 0.01f );
+			}
+		}
+		
+		bRequestCancelCompilation = false;
+
+		// Restore any status from before the loop - see PushStatus() above.
+		GWarn->PopStatus();
+
+		if( !bCompileStillInProgress )		
+		{
+			// Compilation finished, now we need to grab all of the text from the output pipe
+			ModuleCompileReadPipeText += FPlatformProcess::ReadPipe(ModuleCompileReadPipe);
+
+			// The ReturnCode is -1 only if compilation was cancelled.
+			CompilationResult = ReturnCode != -1 ? (ECompilationResult::Type)ReturnCode : ECompilationResult::OtherCompilationError;
+
+			// If compilation succeeded for all modules, go back to the modules and update their module file names
+			// in case we recompiled the modules to a new unique file name.  This is needed so that when the module
+			// is reloaded after the recompile, we load the new DLL file name, not the old one
+			if(CompilationResult == ECompilationResult::Succeeded)
+			{
+				for( int32 CurModuleIndex = 0; CurModuleIndex < ModulesThatWereBeingRecompiled.Num(); ++CurModuleIndex )
+				{
+					const FModuleToRecompile& CurModule = ModulesThatWereBeingRecompiled[ CurModuleIndex ];
+
+					// Were we asked to assign a new file name for this module?
+					if( !CurModule.NewModuleFilename.IsEmpty() )
+					{
+						// If the compile succeeded, update the module info entry with the new file name for this module
+						OnModuleCompileSucceeded(FName(*CurModule.ModuleName), CurModule.NewModuleFilename);
+					}
+				}
+
+				ModulesThatWereBeingRecompiled.Empty();
+			}
+
+			// We're done with the process handle now
+			ModuleCompileProcessHandle.Close();
+			ModuleCompileProcessHandle.Reset();
+
+#if PLATFORM_LINUX
+			close(*(int *)ModuleCompileReadPipe);
+#else
+			FPlatformProcess::ClosePipe(ModuleCompileReadPipe, 0);
+#endif
+
+			Ar.Log(*ModuleCompileReadPipeText);
+			const FString FinalOutput = ModuleCompileReadPipeText;
+			ModuleCompileReadPipe = NULL;
+			ModuleCompileReadPipeText = TEXT("");
+
+			// No longer compiling modules
+			ModulesBeingCompiled.Empty();
+
+			bCompileSucceeded = CompilationResult == ECompilationResult::Succeeded;
+
+			if ( bFireEvents )
+			{
+				const bool bShowLogOnSuccess = false;
+				ModuleCompilerFinishedEvent.Broadcast(FinalOutput, CompilationResult, !bCompileSucceeded || bShowLogOnSuccess);
+
+				// Fire task completion delegate 
+				RecompileModulesCallback.ExecuteIfBound( true, bCompileSucceeded );
+				RecompileModulesCallback.Unbind();
+			}
+		}
+		else
+		{
+			Ar.Logf(TEXT("Error: CheckForFinishedModuleDLLCompile: Compilation is still in progress"));
+		}
+	}
+	else
+	{
+		Ar.Logf(TEXT("Error: CheckForFinishedModuleDLLCompile: There is no compilation in progress right now"));
+	}
+#endif // PLATFORM_DESKTOP && !IS_MONOLITHIC
+}
+
+void FHotReloadModule::UpdateModuleCompileData(FName ModuleName)
+{
+	// Find or create a compile data object for this module
+	TSharedRef<FModuleCompilationData>* CompileDataPtr = ModuleCompileData.Find(ModuleName);
+	if(CompileDataPtr == nullptr)
+	{
+		CompileDataPtr = &ModuleCompileData.Add(ModuleName, TSharedRef<FModuleCompilationData>(new FModuleCompilationData()));
+	}
+
+	// reset the compile data before updating it
+	FModuleCompilationData& CompileData = CompileDataPtr->Get();
+	CompileData.bHasFileTimeStamp = false;
+	CompileData.FileTimeStamp = FDateTime(0);
+	CompileData.CompileMethod = EModuleCompileMethod::Unknown;
+
+#if !IS_MONOLITHIC && WITH_EDITOR
+	ReadModuleCompilationInfoFromConfig(ModuleName, CompileData);
+
+	FDateTime FileTimeStamp;
+	bool bGotFileTimeStamp = GetModuleFileTimeStamp(ModuleName, FileTimeStamp);
+
+	if (!bGotFileTimeStamp)
+	{
+		// File missing? Reset the cached timestamp and method to defaults and save them.
+		CompileData.bHasFileTimeStamp = false;
+		CompileData.FileTimeStamp = FDateTime(0);
+		CompileData.CompileMethod = EModuleCompileMethod::Unknown;
+		WriteModuleCompilationInfoToConfig(ModuleName, CompileData);
+	}
+	else
+	{
+		if (CompileData.bHasFileTimeStamp)
+		{
+			if (FileTimeStamp > CompileData.FileTimeStamp + HotReloadDefs::TimeStampEpsilon)
+			{
+				// The file is newer than the cached timestamp
+				// The file must have been compiled externally
+				CompileData.FileTimeStamp = FileTimeStamp;
+				CompileData.CompileMethod = EModuleCompileMethod::External;
+				WriteModuleCompilationInfoToConfig(ModuleName, CompileData);
+			}
+		}
+		else
+		{
+			// The cached timestamp and method are default value so this file has no history yet
+			// We can only set its timestamp and save
+			CompileData.bHasFileTimeStamp = true;
+			CompileData.FileTimeStamp = FileTimeStamp;
+			WriteModuleCompilationInfoToConfig(ModuleName, CompileData);
+		}
+	}
+#endif
+}
+
+void FHotReloadModule::ReadModuleCompilationInfoFromConfig(FName ModuleName, FModuleCompilationData& CompileData)
+{
+	FString DateTimeString;
+	if (GConfig->GetString(*HotReloadDefs::CompilationInfoConfigSection, *FString::Printf(TEXT("%s.TimeStamp"), *ModuleName.ToString()), DateTimeString, GEditorUserSettingsIni))
+	{
+		FDateTime TimeStamp;
+		if (!DateTimeString.IsEmpty() && FDateTime::Parse(DateTimeString, TimeStamp))
+		{
+			CompileData.bHasFileTimeStamp = true;
+			CompileData.FileTimeStamp = TimeStamp;
+
+			FString CompileMethodString;
+			if (GConfig->GetString(*HotReloadDefs::CompilationInfoConfigSection, *FString::Printf(TEXT("%s.LastCompileMethod"), *ModuleName.ToString()), CompileMethodString, GEditorUserSettingsIni))
+			{
+				if (CompileMethodString.Equals(HotReloadDefs::CompileMethodRuntime, ESearchCase::IgnoreCase))
+				{
+					CompileData.CompileMethod = EModuleCompileMethod::Runtime;
+				}
+				else if (CompileMethodString.Equals(HotReloadDefs::CompileMethodExternal, ESearchCase::IgnoreCase))
+				{
+					CompileData.CompileMethod = EModuleCompileMethod::External;
+				}
+			}
+		}
+	}
+}
+
+void FHotReloadModule::WriteModuleCompilationInfoToConfig(FName ModuleName, const FModuleCompilationData& CompileData)
+{
+	FString DateTimeString;
+	if (CompileData.bHasFileTimeStamp)
+	{
+		DateTimeString = CompileData.FileTimeStamp.ToString();
+	}
+
+	GConfig->SetString(*HotReloadDefs::CompilationInfoConfigSection, *FString::Printf(TEXT("%s.TimeStamp"), *ModuleName.ToString()), *DateTimeString, GEditorUserSettingsIni);
+
+	FString CompileMethodString = HotReloadDefs::CompileMethodUnknown;
+	if (CompileData.CompileMethod == EModuleCompileMethod::Runtime)
+	{
+		CompileMethodString = HotReloadDefs::CompileMethodRuntime;
+	}
+	else if (CompileData.CompileMethod == EModuleCompileMethod::External)
+	{
+		CompileMethodString = HotReloadDefs::CompileMethodExternal;
+	}
+
+	GConfig->SetString(*HotReloadDefs::CompilationInfoConfigSection, *FString::Printf(TEXT("%s.LastCompileMethod"), *ModuleName.ToString()), *CompileMethodString, GEditorUserSettingsIni);
+}
+
+bool FHotReloadModule::GetModuleFileTimeStamp(FName ModuleName, FDateTime& OutFileTimeStamp) const
+{
+	FString Filename = FModuleManager::Get().GetModuleFilename(ModuleName);
+	if (IFileManager::Get().FileSize(*Filename) > 0)
+	{
+		OutFileTimeStamp = FDateTime(IFileManager::Get().GetTimeStamp(*Filename));
+		return true;
+	}
+	return false;
+}
+
 #undef LOCTEXT_NAMESPACE
\ No newline at end of file
diff --git a/Engine/Source/Developer/ModuleUI/Private/SModuleUI.cpp b/Engine/Source/Developer/ModuleUI/Private/SModuleUI.cpp
--- a/Engine/Source/Developer/ModuleUI/Private/SModuleUI.cpp
+++ b/Engine/Source/Developer/ModuleUI/Private/SModuleUI.cpp
@@ -252,18 +252,18 @@ FReply SModuleUI::FModuleListItem::OnRecompileClicked()
 			}
 		}
 
+		IHotReloadInterface& HotReloadSupport = FModuleManager::LoadModuleChecked<IHotReloadInterface>("HotReload");
 		if( PackagesToRebind.Num() > 0 )
 		{
 			// Perform a hot reload
 			const bool bWaitForCompletion = true;			
-			IHotReloadInterface& HotReloadSupport = FModuleManager::LoadModuleChecked<IHotReloadInterface>("HotReload");
 			HotReloadSupport.RebindPackages(PackagesToRebind, TArray<FName>(), bWaitForCompletion, *GLog);
 		}
 		else
 		{
 			// Perform a regular unload, then reload
 			const bool bReloadAfterRecompile = true;
-			FModuleManager::Get().RecompileModule( ModuleName, bReloadAfterRecompile, *GLog );
+			HotReloadSupport.RecompileModule( ModuleName, bReloadAfterRecompile, *GLog );
 		}
 	}
 
diff --git a/Engine/Source/Editor/Kismet/Kismet.Build.cs b/Engine/Source/Editor/Kismet/Kismet.Build.cs
--- a/Engine/Source/Editor/Kismet/Kismet.Build.cs
+++ b/Engine/Source/Editor/Kismet/Kismet.Build.cs
@@ -43,6 +43,7 @@ public class Kismet : ModuleRules
                 "JsonUtilities",
                 "DerivedDataCache",
 				"DesktopPlatform",
+				"HotReload"
 			}
 			);
 
diff --git a/Engine/Source/Editor/Kismet/Private/BlueprintEditor.cpp b/Engine/Source/Editor/Kismet/Private/BlueprintEditor.cpp
--- a/Engine/Source/Editor/Kismet/Private/BlueprintEditor.cpp
+++ b/Engine/Source/Editor/Kismet/Private/BlueprintEditor.cpp
@@ -73,6 +73,7 @@
 #include "IAnalyticsProvider.h"
 #include "ISCSEditorCustomization.h"
 #include "Editor/UnrealEd/Public/SourceCodeNavigation.h"
+#include "Developer/HotReload/Public/IHotReload.h"
 
 // Blueprint merging
 #include "Merge.h"
@@ -1893,16 +1894,16 @@ void FBlueprintEditor::CreateDefaultCommands()
 		FIsActionChecked::CreateSP(this, &FBlueprintEditor::GetSaveIntermediateBuildProducts));
 	ToolkitCommands->MapAction( FBlueprintEditorCommands::Get().RecompileGraphEditor,
 		FExecuteAction::CreateStatic( &FLocalKismetCallbacks::RecompileGraphEditor_OnClicked ),
-		FCanExecuteAction::CreateStatic( &FLocalKismetCallbacks::CanRecompileModules ));
+		FCanExecuteAction::CreateStatic( &FBlueprintEditor::CanRecompileModules ));
 	ToolkitCommands->MapAction( FBlueprintEditorCommands::Get().RecompileKismetCompiler,
 		FExecuteAction::CreateStatic( &FLocalKismetCallbacks::RecompileKismetCompiler_OnClicked ),
-		FCanExecuteAction::CreateStatic( &FLocalKismetCallbacks::CanRecompileModules ));
+		FCanExecuteAction::CreateStatic( &FBlueprintEditor::CanRecompileModules ));
 	ToolkitCommands->MapAction( FBlueprintEditorCommands::Get().RecompileBlueprintEditor,
 		FExecuteAction::CreateStatic( &FLocalKismetCallbacks::RecompileBlueprintEditor_OnClicked ),
-		FCanExecuteAction::CreateStatic( &FLocalKismetCallbacks::CanRecompileModules ));
+		FCanExecuteAction::CreateStatic( &FBlueprintEditor::CanRecompileModules ));
 	ToolkitCommands->MapAction( FBlueprintEditorCommands::Get().RecompilePersona,
 		FExecuteAction::CreateStatic( &FLocalKismetCallbacks::RecompilePersona_OnClicked ),
-		FCanExecuteAction::CreateStatic( &FLocalKismetCallbacks::CanRecompileModules ));
+		FCanExecuteAction::CreateStatic( &FBlueprintEditor::CanRecompileModules ));
 
 	ToolkitCommands->MapAction(FBlueprintEditorCommands::Get().BeginBlueprintMerge,
 		FExecuteAction::CreateSP(this, &FBlueprintEditor::CreateMergeToolTab),
@@ -6170,6 +6171,12 @@ void FBlueprintEditor::RestoreEditedObjectState()
 	}
 }
 
+bool FBlueprintEditor::CanRecompileModules()
+{
+	// We're not able to recompile if a compile is already in progress!
+	return !IHotReloadModule::Get().IsCurrentlyCompiling();
+}
+
 void FBlueprintEditor::OnEditTunnel()
 {
 	auto GraphEditor = FocusedGraphEdPtr.Pin();
diff --git a/Engine/Source/Editor/Kismet/Private/BlueprintEditorTabFactories.h b/Engine/Source/Editor/Kismet/Private/BlueprintEditorTabFactories.h
--- a/Engine/Source/Editor/Kismet/Private/BlueprintEditorTabFactories.h
+++ b/Engine/Source/Editor/Kismet/Private/BlueprintEditorTabFactories.h
@@ -44,12 +44,6 @@ struct FLocalKismetCallbacks
 		FAssetEditorManager::Get().CloseAllAssetEditors();
 		GEngine->DeferredCommands.Add( TEXT( "Module Recompile Persona" ) );
 	}
-
-	static bool CanRecompileModules()
-	{
-		// We're not able to recompile if a compile is already in progress!
-		return !FModuleManager::Get().IsCurrentlyCompiling();
-	}
 };
 
 /////////////////////////////////////////////////////
diff --git a/Engine/Source/Editor/Kismet/Public/BlueprintEditor.h b/Engine/Source/Editor/Kismet/Public/BlueprintEditor.h
--- a/Engine/Source/Editor/Kismet/Public/BlueprintEditor.h
+++ b/Engine/Source/Editor/Kismet/Public/BlueprintEditor.h
@@ -824,6 +824,9 @@ protected:
 
 private:
 
+	/** Returns true if modules can be recompiled */
+	static bool CanRecompileModules();
+
 	/* User wants to edit tunnel via function editor */
 	void OnEditTunnel();
 
diff --git a/Engine/Source/Editor/LevelEditor/Private/LevelEditorActions.cpp b/Engine/Source/Editor/LevelEditor/Private/LevelEditorActions.cpp
--- a/Engine/Source/Editor/LevelEditor/Private/LevelEditorActions.cpp
+++ b/Engine/Source/Editor/LevelEditor/Private/LevelEditorActions.cpp
@@ -898,9 +898,9 @@ void FLevelEditorActionCallbacks::ReloadLevelEditor_Clicked()
 void FLevelEditorActionCallbacks::RecompileGameCode_Clicked()
 {
 	// Don't allow a recompile while already compiling!
-	if( !FModuleManager::Get().IsCurrentlyCompiling() )
+	IHotReloadInterface& HotReloadSupport = FModuleManager::LoadModuleChecked<IHotReloadInterface>("HotReload");
+	if( !HotReloadSupport.IsCurrentlyCompiling() )
 	{
-		IHotReloadInterface& HotReloadSupport = FModuleManager::LoadModuleChecked<IHotReloadInterface>("HotReload");
 		HotReloadSupport.DoHotReloadFromEditor();
 	}
 }
@@ -908,13 +908,15 @@ void FLevelEditorActionCallbacks::RecompileGameCode_Clicked()
 bool FLevelEditorActionCallbacks::Recompile_CanExecute()
 {
 	// We're not able to recompile if a compile is already in progress!
-	return !FModuleManager::Get().IsCurrentlyCompiling() && !(GEngineVersion.IsPromotedBuild() && FEngineBuildSettings::IsPerforceBuild());
+	IHotReloadInterface& HotReloadSupport = FModuleManager::LoadModuleChecked<IHotReloadInterface>("HotReload");
+	return !HotReloadSupport.IsCurrentlyCompiling() && !(GEngineVersion.IsPromotedBuild() && FEngineBuildSettings::IsPerforceBuild());
 }
 
 bool FLevelEditorActionCallbacks::Reload_CanExecute()
 {
 	// We're not able to reload if a compile is already in progress!
-	return !FModuleManager::Get().IsCurrentlyCompiling();
+	IHotReloadInterface& HotReloadSupport = FModuleManager::LoadModuleChecked<IHotReloadInterface>("HotReload");
+	return !HotReloadSupport.IsCurrentlyCompiling();
 }
 
 void FLevelEditorActionCallbacks::GoToCodeForActor_Clicked()
diff --git a/Engine/Source/Editor/MainFrame/Private/MainFrameModule.cpp b/Engine/Source/Editor/MainFrame/Private/MainFrameModule.cpp
--- a/Engine/Source/Editor/MainFrame/Private/MainFrameModule.cpp
+++ b/Engine/Source/Editor/MainFrame/Private/MainFrameModule.cpp
@@ -604,11 +604,11 @@ void FMainFrameModule::StartupModule( )
 
 	SetLevelNameForWindowTitle(TEXT(""));
 
-	FModuleManager::Get().OnModuleCompilerStarted().AddRaw( this, &FMainFrameModule::HandleLevelEditorModuleCompileStarted );
-	FModuleManager::Get().OnModuleCompilerFinished().AddRaw( this, &FMainFrameModule::HandleLevelEditorModuleCompileFinished );
-
 	// Register to find out about when hot reload completes, so we can show a notification
-	IHotReloadModule::Get().OnHotReload().AddRaw( this, &FMainFrameModule::HandleHotReloadFinished );
+	IHotReloadModule& HotReloadModule = IHotReloadModule::Get();
+	HotReloadModule.OnModuleCompilerStarted().AddRaw( this, &FMainFrameModule::HandleLevelEditorModuleCompileStarted );
+	HotReloadModule.OnModuleCompilerFinished().AddRaw( this, &FMainFrameModule::HandleLevelEditorModuleCompileFinished );
+	HotReloadModule.OnHotReload().AddRaw( this, &FMainFrameModule::HandleHotReloadFinished );
 
 #if WITH_EDITOR
 	ISourceCodeAccessModule& SourceCodeAccessModule = FModuleManager::LoadModuleChecked<ISourceCodeAccessModule>("SourceCodeAccess");
@@ -647,12 +647,12 @@ void FMainFrameModule::ShutdownModule( )
 
 	if( IHotReloadModule::IsAvailable() )
 	{
-		IHotReloadModule::Get().OnHotReload().RemoveAll( this );
+		IHotReloadModule& HotReloadModule = IHotReloadModule::Get();
+		HotReloadModule.OnHotReload().RemoveAll( this );
+		HotReloadModule.OnModuleCompilerStarted().RemoveAll( this );
+		HotReloadModule.OnModuleCompilerFinished().RemoveAll( this );
 	}
 
-	FModuleManager::Get().OnModuleCompilerStarted().RemoveAll( this );
-	FModuleManager::Get().OnModuleCompilerFinished().RemoveAll( this );
-
 #if WITH_EDITOR
 	if(FModuleManager::Get().IsModuleLoaded("SourceCodeAccess"))
 	{
@@ -731,7 +731,7 @@ void FMainFrameModule::HandleLevelEditorModuleCompileStarted( )
 
 void FMainFrameModule::OnCancelCodeCompilationClicked()
 {
-	FModuleManager::Get().RequestStopCompilation();
+	IHotReloadModule::Get().RequestStopCompilation();
 }
 
 void FMainFrameModule::HandleLevelEditorModuleCompileFinished(const FString& LogDump, ECompilationResult::Type CompilationResult, bool bShowLog)
diff --git a/Engine/Source/Editor/MainFrame/Private/MainFrameModule.h b/Engine/Source/Editor/MainFrame/Private/MainFrameModule.h
--- a/Engine/Source/Editor/MainFrame/Private/MainFrameModule.h
+++ b/Engine/Source/Editor/MainFrame/Private/MainFrameModule.h
@@ -2,6 +2,7 @@
 
 #pragma once
 
+#include "HotReloadInterface.h"
 
 /**
  * Editor main frame module
diff --git a/Engine/Source/Editor/UnrealEd/Private/Editor.cpp b/Engine/Source/Editor/UnrealEd/Private/Editor.cpp
--- a/Engine/Source/Editor/UnrealEd/Private/Editor.cpp
+++ b/Engine/Source/Editor/UnrealEd/Private/Editor.cpp
@@ -88,6 +88,8 @@
 // AIMdule
 #include "BehaviorTree/BehaviorTreeManager.h"
 
+#include "HotReloadInterface.h"
+
 DEFINE_LOG_CATEGORY_STATIC(LogEditor, Log, All);
 
 #define LOCTEXT_NAMESPACE "UnrealEd.Editor"
@@ -826,8 +828,12 @@ void UEditorEngine::Tick( float DeltaSeconds, bool bIdleMode )
 	// early in the Tick() to get the callbacks for cvar changes called
 	IConsoleManager::Get().CallAllConsoleVariableSinks();
 
-	// Tick the module manager
-	FModuleManager::Get().Tick();
+	// Tick the hot reload interface
+	IHotReloadInterface* HotReload = IHotReloadInterface::GetPtr();
+	if(HotReload != nullptr)
+	{
+		HotReload->Tick();
+	}
 
 	// Tick the remote config IO manager
 	FRemoteConfigAsyncTaskManager::Get()->Tick();
diff --git a/Engine/Source/Editor/UnrealEd/Private/Kismet2/CompilerResultsLog.cpp b/Engine/Source/Editor/UnrealEd/Private/Kismet2/CompilerResultsLog.cpp
--- a/Engine/Source/Editor/UnrealEd/Private/Kismet2/CompilerResultsLog.cpp
+++ b/Engine/Source/Editor/UnrealEd/Private/Kismet2/CompilerResultsLog.cpp
@@ -6,6 +6,7 @@
 #include "MessageLog.h"
 #include "UObjectToken.h"
 #include "SourceCodeNavigation.h"
+#include "Developer/HotReload/Public/IHotReload.h"
 
 #if WITH_EDITOR
 
@@ -80,12 +81,12 @@ void FCompilerResultsLog::Register()
 	FMessageLogModule& MessageLogModule = FModuleManager::LoadModuleChecked<FMessageLogModule>("MessageLog");
 	MessageLogModule.RegisterLogListing(Name, LOCTEXT("CompilerLog", "Compiler Log"));
 
-	FModuleManager::Get().OnModuleCompilerFinished().AddStatic( &FCompilerResultsLog::GetGlobalModuleCompilerDump );
+	IHotReloadModule::Get().OnModuleCompilerFinished().AddStatic( &FCompilerResultsLog::GetGlobalModuleCompilerDump );
 }
 
 void FCompilerResultsLog::Unregister()
 {
-	FModuleManager::Get().OnModuleCompilerFinished().RemoveStatic( &FCompilerResultsLog::GetGlobalModuleCompilerDump );
+	IHotReloadModule::Get().OnModuleCompilerFinished().RemoveStatic( &FCompilerResultsLog::GetGlobalModuleCompilerDump );
 
 	FMessageLogModule& MessageLogModule = FModuleManager::LoadModuleChecked<FMessageLogModule>("MessageLog");
 	MessageLogModule.UnregisterLogListing(Name);
diff --git a/Engine/Source/Editor/UnrealEd/Private/UnrealEdMisc.cpp b/Engine/Source/Editor/UnrealEd/Private/UnrealEdMisc.cpp
--- a/Engine/Source/Editor/UnrealEd/Private/UnrealEdMisc.cpp
+++ b/Engine/Source/Editor/UnrealEd/Private/UnrealEdMisc.cpp
@@ -34,6 +34,7 @@
 #include "DesktopPlatformModule.h"
 #include "ShaderCompiler.h"
 #include "NavigationBuildingNotification.h"
+#include "HotReloadInterface.h"
 
 #define LOCTEXT_NAMESPACE "UnrealEd"
 
@@ -392,17 +393,21 @@ void FUnrealEdMisc::InitEngineAnalytics()
 		}
 
 		// Record known modules' compilation methods
-		TArray<FModuleStatus> Modules;
-		FModuleManager::Get().QueryModules(Modules);
-		for (auto& Module : Modules)
+		IHotReloadInterface* HotReload = IHotReloadInterface::GetPtr();
+		if(HotReload != nullptr)
 		{
-			// Record only game modules as these are the only ones that should be hot-reloaded
-			if (Module.bIsGameModule)
+			TArray<FModuleStatus> Modules;
+			FModuleManager::Get().QueryModules(Modules);
+			for (auto& Module : Modules)
 			{
-				TArray< FAnalyticsEventAttribute > ModuleAttributes;
-				ModuleAttributes.Add(FAnalyticsEventAttribute(FString("ModuleName"), Module.Name));
-				ModuleAttributes.Add(FAnalyticsEventAttribute(FString("CompilationMethod"), FModuleManager::Get().GetModuleCompileMethod(*Module.Name)));
-				EngineAnalytics.RecordEvent(FString("Editor.Usage.Modules"), ModuleAttributes);
+				// Record only game modules as these are the only ones that should be hot-reloaded
+				if (Module.bIsGameModule)
+				{
+					TArray< FAnalyticsEventAttribute > ModuleAttributes;
+					ModuleAttributes.Add(FAnalyticsEventAttribute(FString("ModuleName"), Module.Name));
+					ModuleAttributes.Add(FAnalyticsEventAttribute(FString("CompilationMethod"), HotReload->GetModuleCompileMethod(*Module.Name)));
+					EngineAnalytics.RecordEvent(FString("Editor.Usage.Modules"), ModuleAttributes);
+				}
 			}
 		}
 	}
diff --git a/Engine/Source/Editor/UnrealEd/Public/Kismet2/CompilerResultsLog.h b/Engine/Source/Editor/UnrealEd/Public/Kismet2/CompilerResultsLog.h
--- a/Engine/Source/Editor/UnrealEd/Public/Kismet2/CompilerResultsLog.h
+++ b/Engine/Source/Editor/UnrealEd/Public/Kismet2/CompilerResultsLog.h
@@ -6,6 +6,7 @@
 
 #include "BlueprintUtilities.h"
 #include "TokenizedMessage.h"
+#include "CompilationResult.h"
 
 /** This class maps from final objects to their original source object, across cloning, autoexpansion, etc... */
 class UNREALED_API FBacktrackMap
diff --git a/Engine/Source/Programs/UnrealHeaderTool/Private/HeaderParser.h b/Engine/Source/Programs/UnrealHeaderTool/Private/HeaderParser.h
--- a/Engine/Source/Programs/UnrealHeaderTool/Private/HeaderParser.h
+++ b/Engine/Source/Programs/UnrealHeaderTool/Private/HeaderParser.h
@@ -4,6 +4,7 @@
 
 #include "ParserHelper.h"
 #include "BaseParser.h"
+#include "CompilationResult.h"
 
 class FClass;
 class FClasses;
diff --git a/Engine/Source/Programs/UnrealHeaderTool/Public/UnrealHeaderTool.h b/Engine/Source/Programs/UnrealHeaderTool/Public/UnrealHeaderTool.h
--- a/Engine/Source/Programs/UnrealHeaderTool/Public/UnrealHeaderTool.h
+++ b/Engine/Source/Programs/UnrealHeaderTool/Public/UnrealHeaderTool.h
@@ -6,6 +6,7 @@
 
 #include "Core.h"
 #include "CoreUObject.h"
+#include "CompilationResult.h"
 
 DECLARE_LOG_CATEGORY_EXTERN(LogCompile, Log, All);
 
diff --git a/Engine/Source/Runtime/Core/Private/Modules/ModuleManager.cpp b/Engine/Source/Runtime/Core/Private/Modules/ModuleManager.cpp
--- a/Engine/Source/Runtime/Core/Private/Modules/ModuleManager.cpp
+++ b/Engine/Source/Runtime/Core/Private/Modules/ModuleManager.cpp
@@ -16,19 +16,6 @@ DEFINE_LOG_CATEGORY_STATIC(LogModuleManager, Log, All);
 
 int32 FModuleManager::FModuleInfo::CurrentLoadOrder = 1;
 
-namespace ModuleManagerDefs
-{
-	const static FString CompilationInfoConfigSection("ModuleFileTracking");
-
-	// These strings should match the values of the enum EModuleCompileMethod in ModuleManager.h
-	// and should be handled in ReadModuleCompilationInfoFromConfig() & WriteModuleCompilationInfoToConfig() below
-	const static FString CompileMethodRuntime("Runtime");
-	const static FString CompileMethodExternal("External");
-	const static FString CompileMethodUnknown("Unknown");
-
-	// Add one minute epsilon to timestamp comparision
-	const static FTimespan TimeStampEpsilon(0, 1, 0);
-}
 
 FModuleManager& FModuleManager::Get()
 {
@@ -58,19 +45,6 @@ FModuleManager::~FModuleManager()
 }
 
 
-void FModuleManager::Tick()
-{
-	// We never want to block on a pending compile when checking compilation status during Tick().  We're
-	// just checking so that we can fire callbacks if and when compilation has finished.
-	const bool bWaitForCompletion = false;
-
-	// Ignored output variables
-	bool bCompileStillInProgress = false;
-	bool bCompileSucceeded = false;
-	FOutputDeviceNull NullOutput;
-	CheckForFinishedModuleDLLCompile( bWaitForCompletion, bCompileStillInProgress, bCompileSucceeded, NullOutput );
-}
-
 void FModuleManager::FindModules(const TCHAR* WildcardWithoutExtension, TArray<FName>& OutModules)
 {
 	// @todo plugins: Try to convert existing use cases to use plugins, and get rid of this function
@@ -552,9 +526,6 @@ void FModuleManager::UnloadModulesAtShutdown()
 	{
 		TSharedRef< FModuleInfo > ModuleInfo( ModuleIt.Value() );
 
-		// Take this opportunity to update and write out the compile data before the editor shuts down
-		UpdateModuleCompileData(ModuleIt.Key());
-
 		// Only if already loaded
 		if( ModuleInfo->Module.IsValid() )
 		{
@@ -711,21 +682,6 @@ bool FModuleManager::Exec( UWorld* Inworld, const TCHAR* Cmd, FOutputDevice& Ar
 
 			return true;
 		}
-
-
-		// Recompile <ModuleName>
-		else if( FParse::Command( &Cmd, TEXT( "Recompile" ) ) )
-		{
-			const FString ModuleNameStr = FParse::Token( Cmd, 0 );
-			if( !ModuleNameStr.IsEmpty() )
-			{
-				const FName ModuleName( *ModuleNameStr );
-				const bool bReloadAfterRecompile = true;
-				RecompileModule( ModuleName, bReloadAfterRecompile, Ar);
-			}
-
-			return true;
-		}
 #endif // !IS_MONOLITHIC
 	}
 #endif // !UE_BUILD_SHIPPING
@@ -783,224 +739,6 @@ void FModuleManager::QueryModules( TArray< FModuleStatus >& OutModuleStatuses )
 	}
 }
 
-FString FModuleManager::GetModuleCompileMethod(FName InModuleName)
-{
-	if (!ModuleCompileData.Contains(InModuleName))
-	{
-		UpdateModuleCompileData(InModuleName);
-	}
-
-	switch(ModuleCompileData.FindChecked(InModuleName).Get().CompileMethod)
-	{
-	case EModuleCompileMethod::External:
-		return ModuleManagerDefs::CompileMethodExternal;
-	case EModuleCompileMethod::Runtime:
-		return ModuleManagerDefs::CompileMethodRuntime;
-	default:
-		return ModuleManagerDefs::CompileMethodUnknown;
-	}
-}
-
-
-bool FModuleManager::RecompileModule( const FName InModuleName, const bool bReloadAfterRecompile, FOutputDevice &Ar )
-{
-#if !IS_MONOLITHIC
-	const bool bShowProgressDialog = true;
-	const bool bShowCancelButton = false;
-
-	FFormatNamedArguments Args;
-	Args.Add( TEXT("CodeModuleName"), FText::FromName( InModuleName ) );
-	const FText StatusUpdate = FText::Format( NSLOCTEXT("ModuleManager", "Recompile_SlowTaskName", "Compiling {CodeModuleName}..."), Args );
-
-	GWarn->BeginSlowTask( StatusUpdate, bShowProgressDialog, bShowCancelButton );
-
-	ModuleCompilerStartedEvent.Broadcast();
-
-	// Update our set of known modules, in case we don't already know about this module
-	AddModule( InModuleName );
-
-	// Only use rolling module names if the module was already loaded into memory.  This allows us to try compiling
-	// the module without actually having to unload it first.
-	const bool bWasModuleLoaded = IsModuleLoaded( InModuleName );
-	const bool bUseRollingModuleNames = bWasModuleLoaded;
-	TSharedRef< FModuleInfo > Module = Modules.FindChecked( InModuleName );
-
-	bool bWasSuccessful = true;
-	FString NewModuleFileNameOnSuccess = Module->Filename;
-	if( bUseRollingModuleNames )
-	{
-		// First, try to compile the module.  If the module is already loaded, we won't unload it quite yet.  Instead
-		// make sure that it compiles successfully.
-
-		// Find a unique file name for the module
-		FString UniqueSuffix;
-		FString UniqueModuleFileName;
-		MakeUniqueModuleFilename( InModuleName, UniqueSuffix, UniqueModuleFileName );
-
-		// If the recompile succeeds, we'll update our cached file name to use the new unique file name
-		// that we setup for the module
-		NewModuleFileNameOnSuccess = UniqueModuleFileName;
-
-		TArray< FModuleToRecompile > ModulesToRecompile;
-		FModuleToRecompile ModuleToRecompile;
-		ModuleToRecompile.ModuleName = InModuleName.ToString();
-		ModuleToRecompile.ModuleFileSuffix = UniqueSuffix;
-		ModuleToRecompile.NewModuleFilename = UniqueModuleFileName;
-		ModulesToRecompile.Add( ModuleToRecompile );
-		bWasSuccessful = RecompileModuleDLLs( ModulesToRecompile, Ar );
-	}
-
-	if( bWasSuccessful )
-	{
-		// Shutdown the module if it's already running
-		if( bWasModuleLoaded )
-		{
-			Ar.Logf( TEXT( "Unloading module before compile." ) );
-			UnloadOrAbandonModuleWithCallback( InModuleName, Ar );
-		}
-
-		if( !bUseRollingModuleNames )
-		{
-			// Try to recompile the DLL
-			TArray< FModuleToRecompile > ModulesToRecompile;
-			FModuleToRecompile ModuleToRecompile;
-			ModuleToRecompile.ModuleName = InModuleName.ToString();
-			ModulesToRecompile.Add( ModuleToRecompile );
-			bWasSuccessful = RecompileModuleDLLs( ModulesToRecompile, Ar );
-		}
-
-		// Reload the module if it was loaded before we recompiled
-		if( bWasSuccessful && bWasModuleLoaded && bReloadAfterRecompile )
-		{
-			Ar.Logf( TEXT( "Reloading module after successful compile." ) );
-			bWasSuccessful = LoadModuleWithCallback( InModuleName, Ar );
-		}
-	}
-
-	GWarn->EndSlowTask();
-	return bWasSuccessful;
-#else
-	return false;
-#endif // !IS_MONOLITHIC
-}
-
-
-bool FModuleManager::IsCurrentlyCompiling() const
-{
-	return ModuleCompileProcessHandle.IsValid();
-}
-
-
-FString FModuleManager::MakeUBTArgumentsForModuleCompiling()
-{
-	FString AdditionalArguments;
-	if ( FPaths::IsProjectFilePathSet() )
-	{
-		// We have to pass FULL paths to UBT
-		FString FullProjectPath = FPaths::ConvertRelativePathToFull(FPaths::GetProjectFilePath());
-
-		// @todo projectdirs: Currently non-Rocket projects that exist under the UE4 root are compiled by UBT with no .uproject file
-		//     name passed in (see bIsProjectTarget in VCProject.cs), which causes intermediate libraries to be saved to the Engine
-		//     intermediate folder instead of the project's intermediate folder.  We're emulating this behavior here for module
-		//     recompiling, so that compiled modules will be able to find their import libraries in the original folder they were compiled.
-		if( FRocketSupport::IsRocket() || !FullProjectPath.StartsWith( FPaths::ConvertRelativePathToFull( FPaths::RootDir() ) ) )
-		{
-			const FString ProjectFilenameWithQuotes = FString::Printf(TEXT("\"%s\""), *FullProjectPath);
-			AdditionalArguments += FString::Printf(TEXT("%s "), *ProjectFilenameWithQuotes);
-		}
-
-		if (FRocketSupport::IsRocket())
-		{
-			AdditionalArguments += TEXT("-rocket ");
-		}
-	}
-
-	return AdditionalArguments;
-}
-
-void FModuleManager::OnModuleCompileSucceeded(FName ModuleName, const FString& NewModuleFilename)
-{
-	TSharedRef< FModuleInfo > ModuleInfo = Modules.FindChecked(ModuleName);
-
-	// If the compile succeeded, update the module info entry with the new file name for this module
-	ModuleInfo->Filename = NewModuleFilename;
-
-#if !IS_MONOLITHIC && WITH_EDITOR
-	// UpdateModuleCompileData() should have been run before compiling so the
-	// data in ModuleInfo should be correct for the pre-compile dll file.
-	FModuleCompilationData& CompileData = ModuleCompileData.FindChecked(ModuleName).Get();
-
-	FDateTime FileTimeStamp;
-	bool bGotFileTimeStamp = GetModuleFileTimeStamp(ModuleName, FileTimeStamp);
-
-	CompileData.bHasFileTimeStamp = bGotFileTimeStamp;
-	CompileData.FileTimeStamp = FileTimeStamp;
-
-	if (CompileData.bHasFileTimeStamp)
-	{
-		CompileData.CompileMethod = EModuleCompileMethod::Runtime;
-	}
-	else
-	{
-		CompileData.CompileMethod = EModuleCompileMethod::Unknown;
-	}
-	WriteModuleCompilationInfoToConfig(ModuleName, CompileData);
-#endif
-}
-
-void FModuleManager::UpdateModuleCompileData(FName ModuleName)
-{
-	// Find or create a compile data object for this module
-	TSharedRef<FModuleCompilationData>* CompileDataPtr = ModuleCompileData.Find(ModuleName);
-	if(CompileDataPtr == nullptr)
-	{
-		CompileDataPtr = &ModuleCompileData.Add(ModuleName, TSharedRef<FModuleCompilationData>(new FModuleCompilationData()));
-	}
-
-	// reset the compile data before updating it
-	FModuleCompilationData& CompileData = CompileDataPtr->Get();
-	CompileData.bHasFileTimeStamp = false;
-	CompileData.FileTimeStamp = FDateTime(0);
-	CompileData.CompileMethod = EModuleCompileMethod::Unknown;
-
-#if !IS_MONOLITHIC && WITH_EDITOR
-	ReadModuleCompilationInfoFromConfig(ModuleName, CompileData);
-
-	FDateTime FileTimeStamp;
-	bool bGotFileTimeStamp = GetModuleFileTimeStamp(ModuleName, FileTimeStamp);
-
-	if (!bGotFileTimeStamp)
-	{
-		// File missing? Reset the cached timestamp and method to defaults and save them.
-		CompileData.bHasFileTimeStamp = false;
-		CompileData.FileTimeStamp = FDateTime(0);
-		CompileData.CompileMethod = EModuleCompileMethod::Unknown;
-		WriteModuleCompilationInfoToConfig(ModuleName, CompileData);
-	}
-	else
-	{
-		if (CompileData.bHasFileTimeStamp)
-		{
-			if (FileTimeStamp > CompileData.FileTimeStamp + ModuleManagerDefs::TimeStampEpsilon)
-			{
-				// The file is newer than the cached timestamp
-				// The file must have been compiled externally
-				CompileData.FileTimeStamp = FileTimeStamp;
-				CompileData.CompileMethod = EModuleCompileMethod::External;
-				WriteModuleCompilationInfoToConfig(ModuleName, CompileData);
-			}
-		}
-		else
-		{
-			// The cached timestamp and method are default value so this file has no history yet
-			// We can only set its timestamp and save
-			CompileData.bHasFileTimeStamp = true;
-			CompileData.FileTimeStamp = FileTimeStamp;
-			WriteModuleCompilationInfoToConfig(ModuleName, CompileData);
-		}
-	}
-#endif
-}
 
 FString FModuleManager::GetModuleFilename(FName ModuleName) const
 {
@@ -1012,68 +750,6 @@ void FModuleManager::SetModuleFilename(FName ModuleName, const FString& Filename
 	Modules.FindChecked(ModuleName)->Filename = Filename;
 }
 
-void FModuleManager::ReadModuleCompilationInfoFromConfig(FName ModuleName, FModuleCompilationData& CompileData)
-{
-	FString DateTimeString;
-	if (GConfig->GetString(*ModuleManagerDefs::CompilationInfoConfigSection, *FString::Printf(TEXT("%s.TimeStamp"), *ModuleName.ToString()), DateTimeString, GEditorUserSettingsIni))
-	{
-		FDateTime TimeStamp;
-		if (!DateTimeString.IsEmpty() && FDateTime::Parse(DateTimeString, TimeStamp))
-		{
-			CompileData.bHasFileTimeStamp = true;
-			CompileData.FileTimeStamp = TimeStamp;
-
-			FString CompileMethodString;
-			if (GConfig->GetString(*ModuleManagerDefs::CompilationInfoConfigSection, *FString::Printf(TEXT("%s.LastCompileMethod"), *ModuleName.ToString()), CompileMethodString, GEditorUserSettingsIni))
-			{
-				if (CompileMethodString.Equals(ModuleManagerDefs::CompileMethodRuntime, ESearchCase::IgnoreCase))
-				{
-					CompileData.CompileMethod = EModuleCompileMethod::Runtime;
-				}
-				else if (CompileMethodString.Equals(ModuleManagerDefs::CompileMethodExternal, ESearchCase::IgnoreCase))
-				{
-					CompileData.CompileMethod = EModuleCompileMethod::External;
-				}
-			}
-		}
-	}
-}
-
-
-void FModuleManager::WriteModuleCompilationInfoToConfig(FName ModuleName, const FModuleCompilationData& CompileData)
-{
-	FString DateTimeString;
-	if (CompileData.bHasFileTimeStamp)
-	{
-		DateTimeString = CompileData.FileTimeStamp.ToString();
-	}
-
-	GConfig->SetString(*ModuleManagerDefs::CompilationInfoConfigSection, *FString::Printf(TEXT("%s.TimeStamp"), *ModuleName.ToString()), *DateTimeString, GEditorUserSettingsIni);
-
-	FString CompileMethodString = ModuleManagerDefs::CompileMethodUnknown;
-	if (CompileData.CompileMethod == EModuleCompileMethod::Runtime)
-	{
-		CompileMethodString = ModuleManagerDefs::CompileMethodRuntime;
-	}
-	else if (CompileData.CompileMethod == EModuleCompileMethod::External)
-	{
-		CompileMethodString = ModuleManagerDefs::CompileMethodExternal;
-	}
-
-	GConfig->SetString(*ModuleManagerDefs::CompilationInfoConfigSection, *FString::Printf(TEXT("%s.LastCompileMethod"), *ModuleName.ToString()), *CompileMethodString, GEditorUserSettingsIni);
-}
-
-
-bool FModuleManager::GetModuleFileTimeStamp(FName ModuleName, FDateTime& OutFileTimeStamp) const
-{
-	FString Filename = GetModuleFilename(ModuleName);
-	if (IFileManager::Get().FileSize(*Filename) > 0)
-	{
-		OutFileTimeStamp = FDateTime(IFileManager::Get().GetTimeStamp(*Filename));
-		return true;
-	}
-	return false;
-}
 
 
 void FModuleManager::FindModulePaths(const TCHAR* NamePattern, TMap<FName, FString> &OutModulePaths) const
@@ -1165,68 +841,6 @@ void FModuleManager::FindModulePathsInDirectory(const FString& InDirectoryName,
 }
 
 
-bool FModuleManager::RecompileModulesAsync( const TArray< FName > ModuleNames, const FRecompileModulesCallback& InRecompileModulesCallback, const bool bWaitForCompletion, FOutputDevice &Ar )
-{
-#if !IS_MONOLITHIC
-	// NOTE: This method of recompiling always using a rolling file name scheme, since we never want to unload before
-	// we start recompiling, and we need the output DLL to be unlocked before we invoke the compiler
-
-	ModuleCompilerStartedEvent.Broadcast();
-
-	TArray< FModuleToRecompile > ModulesToRecompile;
-
-	for( TArray< FName >::TConstIterator CurModuleIt( ModuleNames ); CurModuleIt; ++CurModuleIt )
-	{
-		const FName CurModuleName = *CurModuleIt;
-
-		// Update our set of known modules, in case we don't already know about this module
-		AddModule( CurModuleName );
-
-		TSharedRef< FModuleInfo > Module = Modules.FindChecked( CurModuleName );
-
-		FString NewModuleFileNameOnSuccess = Module->Filename;
-
-		// Find a unique file name for the module
-		FString UniqueSuffix;
-		FString UniqueModuleFileName;
-		MakeUniqueModuleFilename( CurModuleName, UniqueSuffix, UniqueModuleFileName );
-
-		// If the recompile succeeds, we'll update our cached file name to use the new unique file name
-		// that we setup for the module
-		NewModuleFileNameOnSuccess = UniqueModuleFileName;
-
-		FModuleToRecompile ModuleToRecompile;
-		ModuleToRecompile.ModuleName = CurModuleName.ToString();
-		ModuleToRecompile.ModuleFileSuffix = UniqueSuffix;
-		ModuleToRecompile.NewModuleFilename = UniqueModuleFileName;
-		ModulesToRecompile.Add( ModuleToRecompile );
-	}
-
-	// Kick off compilation!
-	const FString AdditionalArguments = MakeUBTArgumentsForModuleCompiling();
-	const bool bFailIfGeneratedCodeChanges = false;
-	bool bWasSuccessful = StartCompilingModuleDLLs( FApp::GetGameName(), ModulesToRecompile, InRecompileModulesCallback, Ar, bFailIfGeneratedCodeChanges, AdditionalArguments );
-	if (bWasSuccessful)
-	{
-		// Go ahead and check for completion right away.  This is really just so that we can handle the case
-		// where the user asked us to wait for the compile to finish before returning.
-		bool bCompileStillInProgress = false;
-		bool bCompileSucceeded = false;
-		FOutputDeviceNull NullOutput;
-		CheckForFinishedModuleDLLCompile( bWaitForCompletion, bCompileStillInProgress, bCompileSucceeded, NullOutput );
-		if( !bCompileStillInProgress && !bCompileSucceeded )
-		{
-			bWasSuccessful = false;
-		}
-	}
-
-	return bWasSuccessful;
-#else
-	return false;
-#endif // !IS_MONOLITHIC
-}
-
-
 void FModuleManager::UnloadOrAbandonModuleWithCallback( const FName InModuleName, FOutputDevice &Ar )
 {
 	TSharedRef< FModuleInfo > Module = Modules.FindChecked( InModuleName );
@@ -1296,274 +910,6 @@ const TCHAR *FModuleManager::GetUBTConfiguration()
 }
 
 
-bool FModuleManager::StartCompilingModuleDLLs(const FString& GameName, const TArray< FModuleToRecompile >& ModuleNames, 
-	const FRecompileModulesCallback& InRecompileModulesCallback, FOutputDevice& Ar, bool bInFailIfGeneratedCodeChanges, 
-	const FString& InAdditionalCmdLineArgs )
-{
-#if PLATFORM_DESKTOP && !IS_MONOLITHIC
-	// Keep track of what we're compiling
-	ModulesBeingCompiled = ModuleNames;
-	ModulesThatWereBeingRecompiled = ModulesBeingCompiled;
-
-	const TCHAR* BuildPlatformName = FPlatformMisc::GetUBTPlatform();
-	const TCHAR* BuildConfigurationName = GetUBTConfiguration();
-
-	RecompileModulesCallback = InRecompileModulesCallback;
-
-	// Pass a module file suffix to UBT if we have one
-	FString ModuleArg;
-	for( int32 CurModuleIndex = 0; CurModuleIndex < ModuleNames.Num(); ++CurModuleIndex )
-	{
-		if( !ModuleNames[ CurModuleIndex ].ModuleFileSuffix.IsEmpty() )
-		{
-			ModuleArg += FString::Printf( TEXT( " -ModuleWithSuffix %s %s" ), *ModuleNames[ CurModuleIndex ].ModuleName, *ModuleNames[ CurModuleIndex ].ModuleFileSuffix );
-		}
-		else
-		{
-			ModuleArg += FString::Printf( TEXT( " -Module %s" ), *ModuleNames[ CurModuleIndex ].ModuleName );
-		}
-		Ar.Logf( TEXT( "Recompiling %s..." ), *ModuleNames[ CurModuleIndex ].ModuleName );
-
-		// prepare the compile info in the FModuleInfo so that it can be compared after compiling
-		FName ModuleFName(*ModuleNames[ CurModuleIndex ].ModuleName);
-		UpdateModuleCompileData(ModuleFName);
-	}
-
-	FString ExtraArg;
-#if UE_EDITOR
-	// NOTE: When recompiling from the editor, we're passed the game target name, not the editor target name, but we'll
-	//       pass "-editorrecompile" to UBT which tells UBT to figure out the editor target to use for this game, since
-	//       we can't possibly know what the target is called from within the engine code.
-	ExtraArg = TEXT( "-editorrecompile " );
-#endif
-
-	if (bInFailIfGeneratedCodeChanges)
-	{
-		// Additional argument to let UHT know that we can only compile the module if the generated code didn't change
-		ExtraArg += TEXT( "-FailIfGeneratedCodeChanges " );
-	}
-
-	FString CmdLineParams = FString::Printf( TEXT( "%s%s %s %s %s%s" ), 
-		*GameName, *ModuleArg, 
-		BuildPlatformName, BuildConfigurationName, 
-		*ExtraArg, *InAdditionalCmdLineArgs );
-
-	const bool bInvocationSuccessful = InvokeUnrealBuildToolForCompile(CmdLineParams, Ar);
-	if ( !bInvocationSuccessful )
-	{
-		// No longer compiling modules
-		ModulesBeingCompiled.Empty();
-
-		ModuleCompilerFinishedEvent.Broadcast(FString(), ECompilationResult::OtherCompilationError, false);
-
-		// Fire task completion delegate 
-		
-		RecompileModulesCallback.ExecuteIfBound( false, false );
-		RecompileModulesCallback.Unbind();
-	}
-
-	return bInvocationSuccessful;
-#else
-	return false;
-#endif
-}
-
-bool FModuleManager::InvokeUnrealBuildToolForCompile(const FString& InCmdLineParams, FOutputDevice &Ar)
-{
-#if PLATFORM_DESKTOP && !IS_MONOLITHIC
-
-	// Make sure we're not already compiling something!
-	check(!IsCurrentlyCompiling());
-
-	// Setup output redirection pipes, so that we can harvest compiler output and display it ourselves
-#if PLATFORM_LINUX
-	int pipefd[2];
-	pipe(pipefd);
-	void* PipeRead = &pipefd[0];
-	void* PipeWrite = &pipefd[1];
-#else
-	void* PipeRead = NULL;
-	void* PipeWrite = NULL;
-#endif
-
-	verify(FPlatformProcess::CreatePipe(PipeRead, PipeWrite));
-	ModuleCompileReadPipeText = TEXT("");
-
-	FProcHandle ProcHandle = FUBTInvoker::InvokeUnrealBuildToolAsync(InCmdLineParams, Ar, PipeRead, PipeWrite);
-
-	// We no longer need the Write pipe so close it.
-	// We DO need the Read pipe however...
-#if PLATFORM_LINUX
-	close(*(int*)PipeWrite);
-#else
-	FPlatformProcess::ClosePipe(0, PipeWrite);
-#endif
-
-	if (!ProcHandle.IsValid())
-	{
-		// We're done with the process handle now
-		ModuleCompileProcessHandle.Reset();
-		ModuleCompileReadPipe = NULL;
-	}
-	else
-	{
-		ModuleCompileProcessHandle = ProcHandle;
-		ModuleCompileReadPipe = PipeRead;
-	}
-
-	return ProcHandle.IsValid();
-#else
-	return false;
-#endif // PLATFORM_DESKTOP && !IS_MONOLITHIC
-}
-
-
-void FModuleManager::CheckForFinishedModuleDLLCompile( const bool bWaitForCompletion, bool& bCompileStillInProgress, bool& bCompileSucceeded, FOutputDevice& Ar, const FText& SlowTaskOverrideText, bool bFireEvents )
-{
-#if PLATFORM_DESKTOP && !IS_MONOLITHIC
-	bCompileStillInProgress = false;
-	ECompilationResult::Type CompilationResult = ECompilationResult::OtherCompilationError;
-
-	// Is there a compilation in progress?
-	if( IsCurrentlyCompiling() )
-	{
-		bCompileStillInProgress = true;
-
-		// Ensure slow task messages are seen.
-		GWarn->PushStatus();
-
-		// Update the slow task dialog if we were summoned from a synchronous recompile path
-		if (GIsSlowTask)
-		{
-			if ( !SlowTaskOverrideText.IsEmpty() )
-			{
-				GWarn->StatusUpdate(-1, -1, SlowTaskOverrideText);
-			}
-			else
-			{
-				FText StatusUpdate;
-				if ( ModulesBeingCompiled.Num() > 0 )
-				{
-					FFormatNamedArguments Args;
-					Args.Add( TEXT("CodeModuleName"), FText::FromString( ModulesBeingCompiled[0].ModuleName ) );
-					StatusUpdate = FText::Format( NSLOCTEXT("FModuleManager", "CompileSpecificModuleStatusMessage", "{CodeModuleName}: Compiling modules..."), Args );
-				}
-				else
-				{
-					StatusUpdate = NSLOCTEXT("FModuleManager", "CompileStatusMessage", "Compiling modules...");
-				}
-
-				GWarn->StatusUpdate(-1, -1, StatusUpdate);
-			}
-		}
-
-		// Check to see if the compile has finished yet
-		int32 ReturnCode = -1;
-		while (bCompileStillInProgress)
-		{
-			if( FPlatformProcess::GetProcReturnCode( ModuleCompileProcessHandle, &ReturnCode ) )
-			{
-				bCompileStillInProgress = false;
-			}
-			
-			if (bRequestCancelCompilation)
-			{
-				FPlatformProcess::TerminateProc(ModuleCompileProcessHandle);
-				bCompileStillInProgress = bRequestCancelCompilation = false;
-			}
-
-			if( bCompileStillInProgress )
-			{
-				ModuleCompileReadPipeText += FPlatformProcess::ReadPipe(ModuleCompileReadPipe);
-
-				if( !bWaitForCompletion )
-				{
-					// We haven't finished compiling, but we were asked to return immediately
-
-					break;
-				}
-
-				// Give up a small timeslice if we haven't finished recompiling yet
-				FPlatformProcess::Sleep( 0.01f );
-			}
-		}
-		
-		bRequestCancelCompilation = false;
-
-		// Restore any status from before the loop - see PushStatus() above.
-		GWarn->PopStatus();
-
-		if( !bCompileStillInProgress )		
-		{
-			// Compilation finished, now we need to grab all of the text from the output pipe
-			ModuleCompileReadPipeText += FPlatformProcess::ReadPipe(ModuleCompileReadPipe);
-
-			// The ReturnCode is -1 only if compilation was cancelled.
-			CompilationResult = ReturnCode != -1 ? (ECompilationResult::Type)ReturnCode : ECompilationResult::OtherCompilationError;
-
-			// If compilation succeeded for all modules, go back to the modules and update their module file names
-			// in case we recompiled the modules to a new unique file name.  This is needed so that when the module
-			// is reloaded after the recompile, we load the new DLL file name, not the old one
-			if(CompilationResult == ECompilationResult::Succeeded)
-			{
-				for( int32 CurModuleIndex = 0; CurModuleIndex < ModulesThatWereBeingRecompiled.Num(); ++CurModuleIndex )
-				{
-					const FModuleToRecompile& CurModule = ModulesThatWereBeingRecompiled[ CurModuleIndex ];
-
-					// Were we asked to assign a new file name for this module?
-					if( !CurModule.NewModuleFilename.IsEmpty() )
-					{
-						// If the compile succeeded, update the module info entry with the new file name for this module
-						OnModuleCompileSucceeded(FName(*CurModule.ModuleName), CurModule.NewModuleFilename);
-					}
-				}
-
-				ModulesThatWereBeingRecompiled.Empty();
-			}
-
-			// We're done with the process handle now
-			ModuleCompileProcessHandle.Close();
-			ModuleCompileProcessHandle.Reset();
-
-#if PLATFORM_LINUX
-			close(*(int *)ModuleCompileReadPipe);
-#else
-			FPlatformProcess::ClosePipe(ModuleCompileReadPipe, 0);
-#endif
-
-			Ar.Log(*ModuleCompileReadPipeText);
-			const FString FinalOutput = ModuleCompileReadPipeText;
-			ModuleCompileReadPipe = NULL;
-			ModuleCompileReadPipeText = TEXT("");
-
-			// No longer compiling modules
-			ModulesBeingCompiled.Empty();
-
-			bCompileSucceeded = CompilationResult == ECompilationResult::Succeeded;
-
-			if ( bFireEvents )
-			{
-				const bool bShowLogOnSuccess = false;
-				ModuleCompilerFinishedEvent.Broadcast(FinalOutput, CompilationResult, !bCompileSucceeded || bShowLogOnSuccess);
-
-				// Fire task completion delegate 
-				RecompileModulesCallback.ExecuteIfBound( true, bCompileSucceeded );
-				RecompileModulesCallback.Unbind();
-			}
-		}
-		else
-		{
-			Ar.Logf(TEXT("Error: CheckForFinishedModuleDLLCompile: Compilation is still in progress"));
-		}
-	}
-	else
-	{
-		Ar.Logf(TEXT("Error: CheckForFinishedModuleDLLCompile: There is no compilation in progress right now"));
-	}
-#endif // PLATFORM_DESKTOP && !IS_MONOLITHIC
-}
-
-
 bool FModuleManager::CheckModuleCompatibility(const TCHAR* Filename)
 {
 	int32 ModuleApiVersion = FPlatformProcess::GetDllApiVersion(Filename);
@@ -1575,20 +921,6 @@ bool FModuleManager::CheckModuleCompatibility(const TCHAR* Filename)
 	return true;
 }
 
-bool FModuleManager::RecompileModuleDLLs( const TArray< FModuleToRecompile >& ModuleNames, FOutputDevice& Ar )
-{
-	bool bCompileSucceeded = false;
-#if !IS_MONOLITHIC
-	const FString AdditionalArguments = MakeUBTArgumentsForModuleCompiling();
-	if( StartCompilingModuleDLLs( FApp::GetGameName(), ModuleNames, FRecompileModulesCallback(), Ar, true, AdditionalArguments ) )
-	{
-		const bool bWaitForCompletion = true;	// Always wait
-		bool bCompileStillInProgress = false;
-		CheckForFinishedModuleDLLCompile( bWaitForCompletion, bCompileStillInProgress, bCompileSucceeded, Ar );
-	}
-#endif
-	return bCompileSucceeded;
-}
 
 
 void FModuleManager::StartProcessingNewlyLoadedObjects()
diff --git a/Engine/Source/Runtime/Core/Public/Misc/CompilationResult.h b/Engine/Source/Runtime/Core/Public/Misc/CompilationResult.h
new file mode 100644
--- /dev/null
+++ b/Engine/Source/Runtime/Core/Public/Misc/CompilationResult.h
@@ -0,0 +1,50 @@
+// Copyright 1998-2014 Epic Games, Inc. All Rights Reserved.
+
+#pragma once
+
+/*=============================================================================
+	CompilationResult.h:  Defines an enum for return codes between build tools.
+=============================================================================*/
+
+/**
+ * Enumerates possible results of a compilation operation.
+ *
+ * This enum has to be compatible with the one defined in the
+ * UE4\Engine\Source\Programs\UnrealBuildTool\System\ExternalExecution.cs file
+ * to keep communication between UHT, UBT and Editor compiling processes valid.
+ */
+namespace ECompilationResult
+{
+	enum Type
+	{
+		/** Compilation succeeded */
+		Succeeded = 0,
+		/** Compilation failed because generated code changed which was not supported */
+		FailedDueToHeaderChange = 1,
+		/** Compilation failed due to compilation errors */
+		OtherCompilationError,
+		/** Compilation is not supported in the current build */
+		Unsupported,
+		/** Unknown error */
+		Unknown
+	};
+
+	/**
+	* Converts ECompilationResult enum to string.
+	*/
+	static FORCEINLINE const TCHAR* ToString(ECompilationResult::Type Result)
+	{
+		switch (Result)
+		{
+		case ECompilationResult::Succeeded:
+			return TEXT("Succeeded");
+		case ECompilationResult::FailedDueToHeaderChange:
+			return TEXT("FailedDueToHeaderChange");
+		case ECompilationResult::OtherCompilationError:
+			return TEXT("OtherCompilationError");
+		case ECompilationResult::Unsupported:
+			return TEXT("Unsupported");
+		};
+		return TEXT("Unknown");
+	}
+}
diff --git a/Engine/Source/Runtime/Core/Public/Modules/ModuleManager.h b/Engine/Source/Runtime/Core/Public/Modules/ModuleManager.h
--- a/Engine/Source/Runtime/Core/Public/Modules/ModuleManager.h
+++ b/Engine/Source/Runtime/Core/Public/Modules/ModuleManager.h
@@ -35,61 +35,6 @@ enum class EModuleLoadResult
 };
 
 
-/**
- * Enumerates possible results of a compilation operation.
- *
- * This enum has to be compatible with the one defined in the
- * UE4\Engine\Source\Programs\UnrealBuildTool\System\ExternalExecution.cs file
- * to keep communication between UHT, UBT and Editor compiling processes valid.
- */
-namespace ECompilationResult
-{
-	enum Type
-	{
-		/** Compilation succeeded */
-		Succeeded = 0,
-		/** Compilation failed because generated code changed which was not supported */
-		FailedDueToHeaderChange = 1,
-		/** Compilation failed due to compilation errors */
-		OtherCompilationError,
-		/** Compilation is not supported in the current build */
-		Unsupported,
-		/** Unknown error */
-		Unknown
-	};
-
-	/**
-	* Converts ECompilationResult enum to string.
-	*/
-	static FORCEINLINE const TCHAR* ToString(ECompilationResult::Type Result)
-	{
-		switch (Result)
-		{
-		case ECompilationResult::Succeeded:
-			return TEXT("Succeeded");
-		case ECompilationResult::FailedDueToHeaderChange:
-			return TEXT("FailedDueToHeaderChange");
-		case ECompilationResult::OtherCompilationError:
-			return TEXT("OtherCompilationError");
-		case ECompilationResult::Unsupported:
-			return TEXT("Unsupported");
-		};
-		return TEXT("Unknown");
-	}
-}
-
-
-/**
- * Enumerates compilation methods for modules.
- */
-enum class EModuleCompileMethod
-{
-	Runtime,
-	External,
-	Unknown
-};
-
-
 /**
  * Enumerates reasons for modules to change.
  *
@@ -268,14 +213,6 @@ public:
 	 */
 	void QueryModules( TArray<FModuleStatus>& OutModuleStatuses );
 
-	/**
-	 * Queries the compilation method for a given module.
-	 *
-	 * @param InModuleName Module to query the name of
-	 * @return A string describing the method used to compile the module.
-	 */
-	FString GetModuleCompileMethod(FName InModuleName);
-
 	/**
 	 * Unloads a specific module
 	 *
@@ -398,11 +335,6 @@ public:
 		return Modules.Num();
 	}
 
-	/**
-	 * Module manager ticking is only used to check for asynchronously compiled modules that may need to be reloaded
-	 */
-	void Tick( );
-
 	/**
 	 * Unloads modules during the shutdown process.
 	 *
@@ -410,47 +342,6 @@ public:
 	 */
 	void UnloadModulesAtShutdown( );
 
-	/**
-	 * Tries to recompile the specified module.  If the module is loaded, it will first be unloaded (then reloaded after,
-	 * if the recompile was successful.)
-	 *
-	 * @param InModuleName Name of the module to recompile.
-	 * @param bReloadAfterRecompile If true, the module will automatically be reloaded after a successful compile.  Otherwise, you'll need to load it yourself after.
-	 * @param Ar Output device for logging compilation status.
-	 * @return	Returns true if the module was successfully recompiled (and reloaded, if it was previously loaded).
-	 */
-	bool RecompileModule( const FName InModuleName, const bool bReloadAfterRecompile, FOutputDevice &Ar );
-
-	/** @return	Returns true if an asynchronous compile is currently in progress */
-	bool IsCurrentlyCompiling() const;
-
-	/**
-	 * Declares a type of delegates that is executed after a module recompile has finished.
-	 *
-	 * The first argument signals whether compilation has finished.
-	 * The second argument shows whether compilation was successful or not.
-	 */
-	DECLARE_DELEGATE_TwoParams( FRecompileModulesCallback, bool, bool );
-
-	/**
-	 * Tries to recompile the specified modules in the background.  When recompiling finishes, the specified callback
-	 * delegate will be triggered, passing along a bool that tells you whether the compile action succeeded.  This
-	 * function never tries to unload modules or to reload the modules after they finish compiling.  You should do
-	 * that yourself in the recompile completion callback!
-	 *
-	 * @param ModuleNames Names of the modules to recompile
-	 * @param RecompileModulesCallback Callback function to execute after compilation finishes (whether successful or not.)
-	 * @param bWaitForCompletion True if the function should not return until recompilation attempt has finished and callbacks have fired
-	 * @param Ar Output device for logging compilation status
-	 * @return	True if the recompile action was kicked off successfully.  If this returns false, then the recompile callback will never fire.  In the case where bWaitForCompletion=false, this will also return false if the compilation failed for any reason.
-	 */
-	bool RecompileModulesAsync( const TArray< FName > ModuleNames, const FRecompileModulesCallback& RecompileModulesCallback, const bool bWaitForCompletion, FOutputDevice &Ar );
-
-	/** Request that any current compilation operation be abandoned. */
-	void RequestStopCompilation()
-	{
-		bRequestCancelCompilation = true;
-	}
 
 	/** Delegate that's used by the module manager to initialize a registered module that we statically linked with (monolithic only) */
 	DECLARE_DELEGATE_RetVal( IModuleInterface*, FInitializeStaticallyLinkedModule )
@@ -506,12 +397,6 @@ public:
 	 */
 	static const TCHAR *GetUBTConfiguration( );
 
-	/** Called for successfully re-complied module */
-	void OnModuleCompileSucceeded(FName ModuleName, const FString& NewModuleFilename);
-
-	/** Called when the compile data for a module need to be update in memory and written to config */
-	void UpdateModuleCompileData(FName ModuleName);
-
 	/** Gets the filename for a module */
 	FString GetModuleFilename(FName ModuleName) const;
 
@@ -534,31 +419,6 @@ public:
 		return ModulesChangedEvent;
 	}
 
-	/**
-	 * Gets an event delegate that is executed when compilation of a module has started.
-	 *
-	 * @return The event delegate.
-	 */
-	DECLARE_EVENT(FModuleManager, FModuleCompilerStartedEvent);
-	FModuleCompilerStartedEvent& OnModuleCompilerStarted( )
-	{
-		return ModuleCompilerStartedEvent;
-	}
-
-	/**
-	 * Gets an event delegate that is executed when compilation of a module has finished.
-	 *
-	 * The first parameter is the result of the compilation operation.
-	 * The second parameter determines whether the log should be shown.
-	 *
-	 * @return The event delegate.
-	 */
-	DECLARE_EVENT_ThreeParams(FModuleManager, FModuleCompilerFinishedEvent, const FString&, ECompilationResult::Type, bool);
-	FModuleCompilerFinishedEvent& OnModuleCompilerFinished()
-	{
-		return ModuleCompilerFinishedEvent;
-	}
-
 	/**
 	 * Gets a multicast delegate that is executed when any UObjects need processing after a module was loaded.
 	 *
@@ -604,48 +464,10 @@ protected:
 	 */
 	FModuleManager( )
 		: bCanProcessNewlyLoadedObjects(false)
-//		,  ModuleCompileProcessHandle(nullptr)
-		, ModuleCompileReadPipe(nullptr)
-		, bRequestCancelCompilation(false)
 	{ }
 
 protected:
 
-	/**
-	 * Helper structure to hold on to module state while asynchronously recompiling DLLs
-	 */
-	struct FModuleToRecompile
-	{
-		/** Name of the module */
-		FString ModuleName;
-
-		/** Desired module file name suffix, or empty string if not needed */
-		FString ModuleFileSuffix;
-
-		/** The module file name to use after a compilation succeeds, or an empty string if not changing */
-		FString NewModuleFilename;
-	};
-
-	/**
-	 * Helper structure to store the compile time and method for a module
-	 */
-	struct FModuleCompilationData
-	{
-		/** Has a timestamp been set for the .dll file */
-		bool bHasFileTimeStamp;
-
-		/** Last known timestamp for the .dll file */
-		FDateTime FileTimeStamp;
-
-		/** Last known compilation method of the .dll file */
-		EModuleCompileMethod CompileMethod;
-
-		FModuleCompilationData()
-			: bHasFileTimeStamp(false)
-			, CompileMethod(EModuleCompileMethod::Unknown)
-		{ }
-	};
-
 	/**
 	 * Information about a single module (may or may not be loaded.)
 	 */
@@ -687,55 +509,16 @@ protected:
 	/** Type definition for maps of module names to module infos. */
 	typedef TMap<FName, TSharedRef<FModuleInfo>> FModuleMap;
 
-	/**
-	 * Tries to recompile the specified DLL using UBT. Does not interact with modules. This is a low level routine.
-	 *
-	 * @param ModuleNames List of modules to recompile, including the module name and optional file suffix.
-	 * @param Ar Output device for logging compilation status.
-	 */
-	bool RecompileModuleDLLs( const TArray< FModuleToRecompile >& ModuleNames, FOutputDevice &Ar );
-
+public:
 	/**
 	 * Generates a unique file name for the specified module name by adding a random suffix and checking for file collisions.
 	 */
 	void MakeUniqueModuleFilename( const FName InModuleName, FString& UniqueSuffix, FString& UniqueModuleFileName );
 
-	/** @returns Static: Returns arguments to pass to UnrealBuildTool when compiling modules */
-	FString MakeUBTArgumentsForModuleCompiling();
-
-	/** 
-	 *	Starts compiling DLL files for one or more modules.
-	 *
-	 *	@param GameName The name of the game.
-	 *	@param ModuleNames The list of modules to compile.
-	 *	@param InRecompileModulesCallback Callback function to make when module recompiles.
-	 *	@param Ar
-	 *	@param bInFailIfGeneratedCodeChanges If true, fail the compilation if generated headers change.
-	 *	@param InAdditionalCmdLineArgs Additional arguments to pass to UBT.
-	 *	@return true if successful, false otherwise.
-	 */
-	bool StartCompilingModuleDLLs( const FString& GameName, const TArray< FModuleToRecompile >& ModuleNames, 
-		const FRecompileModulesCallback& InRecompileModulesCallback, FOutputDevice& Ar, bool bInFailIfGeneratedCodeChanges, 
-		const FString& InAdditionalCmdLineArgs = FString() );
-
-	/** Launches UnrealBuildTool with the specified command line parameters */
-	bool InvokeUnrealBuildToolForCompile(const FString& InCmdLineParams, FOutputDevice &Ar);	
-
-	/** Checks to see if a pending compilation action has completed and optionally waits for it to finish.  If completed, fires any appropriate callbacks and reports status provided bFireEvents is true. */
-	void CheckForFinishedModuleDLLCompile( const bool bWaitForCompletion, bool& bCompileStillInProgress, bool& bCompileSucceeded, FOutputDevice& Ar, const FText& SlowTaskOverrideTest = FText::GetEmpty(), bool bFireEvents = true );
-
+private:
 	/** Compares file versions between the current executing engine version and the specified dll */
 	static bool CheckModuleCompatibility(const TCHAR *Filename);
 
-	/** Called when a new module is added to the manager to get the saved compile data from config */
-	static void ReadModuleCompilationInfoFromConfig(FName ModuleName, FModuleCompilationData& CompileData);
-
-	/** Saves the module's compile data to config */
-	static void WriteModuleCompilationInfoToConfig(FName ModuleName, const FModuleCompilationData& CompileData);
-
-	/** Access the module's file and read the timestamp from the file system. Returns true if the timestamp was read successfully. */
-	bool GetModuleFileTimeStamp(FName ModuleName, FDateTime& OutFileTimeStamp) const;
-
 	/** Finds modules matching a given name wildcard. */
 	void FindModulePaths(const TCHAR *NamePattern, TMap<FName, FString> &OutModulePaths) const;
 
@@ -758,43 +541,13 @@ private:
 	    our set of known modules changes */
 	FModulesChangedEvent ModulesChangedEvent;
 	
-	/** Multicast delegate which will broadcast a notification when the compiler starts */
-	FModuleCompilerStartedEvent ModuleCompilerStartedEvent;
-	
-	/** Multicast delegate which will broadcast a notification when the compiler finishes */
-	FModuleCompilerFinishedEvent ModuleCompilerFinishedEvent;
-
 	/** Multicast delegate called to process any new loaded objects. */
 	FSimpleMulticastDelegate ProcessLoadedObjectsCallback;
 
-	/** Last known compilation data for each module */
-	TMap<FName, TSharedRef<FModuleCompilationData>> ModuleCompileData;
-
-	/** When compiling a module using an external application, stores the handle to the process that is running */
-	FProcHandle ModuleCompileProcessHandle;
-
-	/** When compiling a module using an external application, this is the process read pipe handle */
-	void* ModuleCompileReadPipe;
-
-	/** When compiling a module using an external application, this is the text that was read from the read pipe handle */
-	FString ModuleCompileReadPipeText;
-
-	/** Callback to execute after an asynchronous recompile has completed (whether successful or not.) */
-	FRecompileModulesCallback RecompileModulesCallback;
-
 	/** When module manager is linked against an application that supports UObjects, this delegate will be primed
 	    at startup to provide information about whether a UObject package is loaded into memory. */
 	FIsPackageLoadedCallback IsPackageLoaded;
 
-	/** Array of modules that we're currently recompiling */
-	TArray< FModuleToRecompile > ModulesBeingCompiled;
-
-	/** Array of modules that we're going to recompile */
-	TArray< FModuleToRecompile > ModulesThatWereBeingRecompiled;
-
-	/** true if we should attempt to cancel the current async compilation */
-	bool bRequestCancelCompilation;
-
 	/** Array of engine binaries directories. */
 	TArray<FString> EngineBinariesDirectories;
 
diff --git a/Engine/Source/Runtime/CoreUObject/Public/Misc/HotReloadInterface.h b/Engine/Source/Runtime/CoreUObject/Public/Misc/HotReloadInterface.h
--- a/Engine/Source/Runtime/CoreUObject/Public/Misc/HotReloadInterface.h
+++ b/Engine/Source/Runtime/CoreUObject/Public/Misc/HotReloadInterface.h
@@ -3,13 +3,55 @@
 #pragma once
 
 #include "ModuleManager.h"
+#include "CoreUObject.h"
+#include "CompilationResult.h"
 
 /**
 * HotReload module interface
 */
-class COREUOBJECT_API IHotReloadInterface : public IModuleInterface
+class IHotReloadInterface : public IModuleInterface
 {
 public:
+	/**
+	 * Tries to gets a pointer to the active HotReload implementation. 
+	 */
+	static inline IHotReloadInterface* GetPtr()
+	{
+		return FModuleManager::GetModulePtr<IHotReloadInterface>("HotReload");
+	}
+
+	/**
+	 * Module manager ticking is only used to check for asynchronously compiled modules that may need to be reloaded
+	 */
+	virtual void Tick() = 0;
+
+	/**
+	 * Save the current state to disk before quitting.
+	 */
+	virtual void SaveConfig() = 0;
+
+	/**
+	 * Queries the compilation method for a given module.
+	 *
+	 * @param InModuleName Module to query the name of
+	 * @return A string describing the method used to compile the module.
+	 */
+	virtual FString GetModuleCompileMethod(FName InModuleName) = 0;
+
+	/**
+	 * Recompiles a single module
+	 */
+	virtual bool RecompileModule( const FName InModuleName, const bool bReloadAfterRecompile, FOutputDevice &Ar ) = 0;
+
+	/**
+	 * Returns whether modules are currently being compiled
+	 */
+	virtual bool IsCurrentlyCompiling() const = 0;
+
+	/**
+	 * Request that current compile be stopped
+	 */
+	virtual void RequestStopCompilation() = 0;
 
 	/**
 	* Adds a function to re-map after hot-reload.
@@ -36,5 +78,24 @@ public:
 	 */
 	DECLARE_MULTICAST_DELEGATE_OneParam(FHotReloadEvent, bool /* bWasTriggeredAutomatically */ );
 	virtual FHotReloadEvent& OnHotReload() = 0;
+
+	/**
+	 * Gets an event delegate that is executed when compilation of a module has started.
+	 *
+	 * @return The event delegate.
+	 */
+	DECLARE_MULTICAST_DELEGATE(FModuleCompilerStartedEvent);
+	virtual FModuleCompilerStartedEvent& OnModuleCompilerStarted() = 0;
+
+	/**
+	 * Gets an event delegate that is executed when compilation of a module has finished.
+	 *
+	 * The first parameter is the result of the compilation operation.
+	 * The second parameter determines whether the log should be shown.
+	 *
+	 * @return The event delegate.
+	 */
+	DECLARE_MULTICAST_DELEGATE_ThreeParams(FModuleCompilerFinishedEvent, const FString&, ECompilationResult::Type, bool);
+	virtual FModuleCompilerFinishedEvent& OnModuleCompilerFinished() = 0;
 };
 
diff --git a/Engine/Source/Runtime/Engine/Private/GameEngine.cpp b/Engine/Source/Runtime/Engine/Private/GameEngine.cpp
--- a/Engine/Source/Runtime/Engine/Private/GameEngine.cpp
+++ b/Engine/Source/Runtime/Engine/Private/GameEngine.cpp
@@ -28,6 +28,7 @@
 
 #include "IHeadMountedDisplay.h"
 #include "RendererInterface.h"
+#include "HotReloadInterface.h"
 
 DEFINE_LOG_CATEGORY_STATIC(LogEngine, Log, All);
 
@@ -777,7 +778,11 @@ void UGameEngine::Tick( float DeltaSeconds, bool bIdleMode )
 	}
 
 	// Tick the module manager
-	FModuleManager::Get().Tick();
+	IHotReloadInterface* HotReload = IHotReloadInterface::GetPtr();
+	if(HotReload != nullptr)
+	{
+		HotReload->Tick();
+	}
 
 	if (!IsRunningDedicatedServer() && !IsRunningCommandlet())
 	{
diff --git a/Engine/Source/Runtime/Engine/Private/RendererSupport.cpp b/Engine/Source/Runtime/Engine/Private/RendererSupport.cpp
--- a/Engine/Source/Runtime/Engine/Private/RendererSupport.cpp
+++ b/Engine/Source/Runtime/Engine/Private/RendererSupport.cpp
@@ -12,6 +12,7 @@
 #include "Slate.h"
 #include "EngineModule.h"
 #include "RendererInterface.h"
+#include "HotReloadInterface.h"
 
 /** Clears and optionally backs up all references to renderer module classes in other modules, particularly engine. */
 void ClearReferencesToRendererModuleClasses(
@@ -83,27 +84,31 @@ void ClearReferencesToRendererModuleClasses(
 /** Recompiles the renderer module, retrying until successful. */
 void RecompileRendererModule()
 {
-	const FName RendererModuleName = TEXT("Renderer");
-	// Unload first so that RecompileModule will not using a rolling module name
-	verify(FModuleManager::Get().UnloadModule(RendererModuleName));
-
-	bool bCompiledSuccessfully = false;
-	do 
+	IHotReloadInterface* HotReload = IHotReloadInterface::GetPtr();
+	if(HotReload != nullptr)
 	{
-		bCompiledSuccessfully = FModuleManager::Get().RecompileModule(RendererModuleName, false, *GLog);
+		const FName RendererModuleName = TEXT("Renderer");
+		// Unload first so that RecompileModule will not using a rolling module name
+		verify(FModuleManager::Get().UnloadModule(RendererModuleName));
 
-		if (!bCompiledSuccessfully)
+		bool bCompiledSuccessfully = false;
+		do 
 		{
-			// Pop up a blocking dialog if there were compilation errors
-			// Compiler output will be in the log
-			FPlatformMisc::MessageBoxExt(EAppMsgType::Ok, *FText::Format(
-				NSLOCTEXT("UnrealEd", "Error_RetryCompilation", "C++ compilation of module {0} failed!  Details in the log.  \r\nFix the error then click Ok to retry."),
-				FText::FromName(RendererModuleName)).ToString(), TEXT("Error"));
-		}
-	} 
-	while (!bCompiledSuccessfully);
+			bCompiledSuccessfully = HotReload->RecompileModule(RendererModuleName, false, *GLog);
+
+			if (!bCompiledSuccessfully)
+			{
+				// Pop up a blocking dialog if there were compilation errors
+				// Compiler output will be in the log
+				FPlatformMisc::MessageBoxExt(EAppMsgType::Ok, *FText::Format(
+					NSLOCTEXT("UnrealEd", "Error_RetryCompilation", "C++ compilation of module {0} failed!  Details in the log.  \r\nFix the error then click Ok to retry."),
+					FText::FromName(RendererModuleName)).ToString(), TEXT("Error"));
+			}
+		} 
+		while (!bCompiledSuccessfully);
 
-	verify(FModuleManager::Get().LoadModule(RendererModuleName, true).IsValid());
+		verify(FModuleManager::Get().LoadModule(RendererModuleName, true).IsValid());
+	}
 }
 
 /** Restores systems that need references to classes in the renderer module. */
diff --git a/Engine/Source/Runtime/Launch/Private/LaunchEngineLoop.cpp b/Engine/Source/Runtime/Launch/Private/LaunchEngineLoop.cpp
--- a/Engine/Source/Runtime/Launch/Private/LaunchEngineLoop.cpp
+++ b/Engine/Source/Runtime/Launch/Private/LaunchEngineLoop.cpp
@@ -47,6 +47,7 @@
 	#include "EngineService.h"
 	#include "ContentStreaming.h"
 	#include "HighResScreenshot.h"
+	#include "HotReloadInterface.h"
 
 #if !UE_SERVER
 	#include "HeadMountedDisplay.h"
@@ -1968,6 +1969,15 @@ void FEngineLoop::Exit()
 	// Tear down the RHI.
 	RHIExit();
 
+#if WITH_ENGINE
+	// Save the hot reload state
+	IHotReloadInterface* HotReload = IHotReloadInterface::GetPtr();
+	if(HotReload != nullptr)
+	{
+		HotReload->SaveConfig();
+	}
+#endif
+
 	// Unload all modules.  Note that this doesn't actually unload the module DLLs (that happens at
 	// process exit by the OS), but it does call ShutdownModule() on all loaded modules in the reverse
 	// order they were loaded in, so that systems can unregister and perform general clean up.
