Fixing recompiling engine modules not working for non code-based game projects

#change Made sure the correct target is passed to UBT when compiling engine/editor modules for non code-based game projects
#change Refactored FLevelEditorActionCallbacks::CanShowSourceCodeActions() code to FHotReloadModule::IsAnyGameModuleLoaded

diff --git a/Engine/Source/Developer/HotReload/Private/HotReload.cpp b/Engine/Source/Developer/HotReload/Private/HotReload.cpp
--- a/Engine/Source/Developer/HotReload/Private/HotReload.cpp
+++ b/Engine/Source/Developer/HotReload/Private/HotReload.cpp
@@ -43,6 +43,7 @@ public:
 	virtual FModuleCompilerStartedEvent& OnModuleCompilerStarted() override { return ModuleCompilerStartedEvent; }
 	virtual FModuleCompilerFinishedEvent& OnModuleCompilerFinished() override { return ModuleCompilerFinishedEvent; }
 	virtual FString GetModuleCompileMethod(FName InModuleName) override;
+	virtual bool IsAnyGameModuleLoaded() const override;
 
 private:
 	/**
@@ -1090,8 +1091,19 @@ bool FHotReloadModule::StartCompilingModuleDLLs(const FString& GameName, const T
 		ExtraArg += TEXT( "-FailIfGeneratedCodeChanges " );
 	}
 
+	// If the're no game modules loaded, then it's not a code-based project and the target
+	// for UBT should be the editor.
+	FString TargetName;
+	if (IsAnyGameModuleLoaded())
+	{
+		TargetName = GameName;
+	}
+	else
+	{
+		TargetName = TEXT("UE4Editor");
+	}
 	FString CmdLineParams = FString::Printf( TEXT( "%s%s %s %s %s%s" ), 
-		*GameName, *ModuleArg, 
+		*TargetName, *ModuleArg, 
 		BuildPlatformName, BuildConfigurationName, 
 		*ExtraArg, *InAdditionalCmdLineArgs );
 
@@ -1424,4 +1436,25 @@ bool FHotReloadModule::GetModuleFileTimeStamp(FName ModuleName, FDateTime& OutFi
 	return false;
 }
 
+bool FHotReloadModule::IsAnyGameModuleLoaded() const
+{
+	// Ask the module manager for a list of currently-loaded gameplay modules
+	TArray< FModuleStatus > ModuleStatuses;
+	FModuleManager::Get().QueryModules(ModuleStatuses);
+
+	for (auto ModuleStatusIt = ModuleStatuses.CreateConstIterator(); ModuleStatusIt; ++ModuleStatusIt)
+	{
+		const FModuleStatus& ModuleStatus = *ModuleStatusIt;
+
+		// We only care about game modules that are currently loaded
+		if (ModuleStatus.bIsLoaded && ModuleStatus.bIsGameModule)
+		{
+			// There is at least one loaded game module.
+			return true;
+		}
+	}
+
+	return false;
+}
+
 #undef LOCTEXT_NAMESPACE
\ No newline at end of file
diff --git a/Engine/Source/Editor/LevelEditor/Private/LevelEditorActions.cpp b/Engine/Source/Editor/LevelEditor/Private/LevelEditorActions.cpp
--- a/Engine/Source/Editor/LevelEditor/Private/LevelEditorActions.cpp
+++ b/Engine/Source/Editor/LevelEditor/Private/LevelEditorActions.cpp
@@ -866,23 +866,9 @@ void FLevelEditorActionCallbacks::MapCheck_Execute()
 
 bool FLevelEditorActionCallbacks::CanShowSourceCodeActions()
 {
-	// Ask the module manager for a list of currently-loaded gameplay modules
-	TArray< FModuleStatus > ModuleStatuses;
-	FModuleManager::Get().QueryModules( ModuleStatuses );
-
-	for( auto ModuleStatusIt = ModuleStatuses.CreateConstIterator(); ModuleStatusIt; ++ModuleStatusIt )
-	{
-		const FModuleStatus& ModuleStatus = *ModuleStatusIt;
-
-		// We only care about game modules that are currently loaded
-		if( ModuleStatus.bIsLoaded && ModuleStatus.bIsGameModule )
-		{
-			// There is at least one loaded game module. Source code actions should be available.
-			return true;
-		}
-	}
-
-	return false;
+	IHotReloadInterface& HotReloadSupport = FModuleManager::LoadModuleChecked<IHotReloadInterface>("HotReload");
+	// If there is at least one loaded game module, source code actions should be available.
+	return HotReloadSupport.IsAnyGameModuleLoaded();
 }
 
 void FLevelEditorActionCallbacks::RecompileLevelEditor_Clicked()
diff --git a/Engine/Source/Runtime/CoreUObject/Public/Misc/HotReloadInterface.h b/Engine/Source/Runtime/CoreUObject/Public/Misc/HotReloadInterface.h
--- a/Engine/Source/Runtime/CoreUObject/Public/Misc/HotReloadInterface.h
+++ b/Engine/Source/Runtime/CoreUObject/Public/Misc/HotReloadInterface.h
@@ -97,5 +97,10 @@ public:
 	 */
 	DECLARE_MULTICAST_DELEGATE_ThreeParams(FModuleCompilerFinishedEvent, const FString&, ECompilationResult::Type, bool);
 	virtual FModuleCompilerFinishedEvent& OnModuleCompilerFinished() = 0;
+
+	/**
+	 * Checks if there's any game modules currently loaded
+	 */
+	virtual bool IsAnyGameModuleLoaded() const = 0;
 };
 
