Making sure bIsHotReloadingFromEditor is set to false in all cases when the compilation fails

#ttp 345519: HotReload: After failing to HotReload once, it would not function again during that session

diff --git a/Engine/Source/Developer/HotReload/Private/HotReload.cpp b/Engine/Source/Developer/HotReload/Private/HotReload.cpp
--- a/Engine/Source/Developer/HotReload/Private/HotReload.cpp
+++ b/Engine/Source/Developer/HotReload/Private/HotReload.cpp
@@ -716,6 +716,7 @@ ECompilationResult::Type FHotReloadModule::RebindPackagesInternal(TArray<UPackag
 			if (bWaitForCompletion)
 			{
 				Ar.Logf(ELogVerbosity::Warning, TEXT("HotReload operation took %4.1fs."), float(FPlatformTime::Seconds() - StartTime));
+				bIsHotReloadingFromEditor = false;
 			}
 			else
 			{
@@ -727,6 +728,7 @@ ECompilationResult::Type FHotReloadModule::RebindPackagesInternal(TArray<UPackag
 		{
 			Ar.Logf(ELogVerbosity::Warning, TEXT("RebindPackages failed because the compiler could not be started."));
 			Result = ECompilationResult::OtherCompilationError;
+			bIsHotReloadingFromEditor = false;
 		}
 	}
 	else
