Move hot reload functionality onto it's own define, instead of being tied into IS_MONOLITHIC. Don't want it enabled for modular shipping games in Rocket.

diff --git a/Engine/Source/Developer/HotReload/Private/HotReload.cpp b/Engine/Source/Developer/HotReload/Private/HotReload.cpp
--- a/Engine/Source/Developer/HotReload/Private/HotReload.cpp
+++ b/Engine/Source/Developer/HotReload/Private/HotReload.cpp
@@ -329,7 +329,7 @@ bool FHotReloadModule::Exec( UWorld* Inworld, const TCHAR* Cmd, FOutputDevice& A
 #if !UE_BUILD_SHIPPING
 	if ( FParse::Command( &Cmd, TEXT( "Module" ) ) )
 	{
-#if !IS_MONOLITHIC
+#if WITH_HOT_RELOAD
 		// Recompile <ModuleName>
 		if( FParse::Command( &Cmd, TEXT( "Recompile" ) ) )
 		{
@@ -343,7 +343,7 @@ bool FHotReloadModule::Exec( UWorld* Inworld, const TCHAR* Cmd, FOutputDevice& A
 
 			return true;
 		}
-#endif // !IS_MONOLITHIC
+#endif // WITH_HOT_RELOAD
 	}
 #endif // !UE_BUILD_SHIPPING
 	return false;
@@ -395,7 +395,7 @@ FString FHotReloadModule::GetModuleCompileMethod(FName InModuleName)
 
 bool FHotReloadModule::RecompileModule( const FName InModuleName, const bool bReloadAfterRecompile, FOutputDevice &Ar )
 {
-#if !IS_MONOLITHIC
+#if WITH_HOT_RELOAD
 	const bool bShowProgressDialog = true;
 	const bool bShowCancelButton = false;
 
@@ -471,7 +471,7 @@ bool FHotReloadModule::RecompileModule( const FName InModuleName, const bool bRe
 	return bWasSuccessful;
 #else
 	return false;
-#endif // !IS_MONOLITHIC
+#endif // WITH_HOT_RELOAD
 }
 
 /** Type hash for a UObject Function Pointer, maybe not a great choice, but it should be sufficient for the needs here. **/
@@ -525,7 +525,7 @@ void FHotReloadModule::DoHotReloadCallback(bool bRecompileFinished, bool bRecomp
 ECompilationResult::Type FHotReloadModule::DoHotReloadInternal(bool bRecompileFinished, bool bRecompileSucceeded, TArray<UPackage*> Packages, TArray< FName > InDependentModules, FOutputDevice &HotReloadAr)
 {
 	ECompilationResult::Type Result = ECompilationResult::Unsupported;
-#if !IS_MONOLITHIC
+#if WITH_HOT_RELOAD
 	if (bRecompileSucceeded)
 	{
 		FFeedbackContext& ErrorsFC = UClass::GetDefaultPropertiesFeedbackContext();
@@ -661,7 +661,7 @@ void FHotReloadModule::RebindPackages(TArray<UPackage*> InPackages, TArray<FName
 ECompilationResult::Type FHotReloadModule::RebindPackagesInternal(TArray<UPackage*> InPackages, TArray<FName> DependentModules, const bool bWaitForCompletion, FOutputDevice &Ar)
 {
 	ECompilationResult::Type Result = ECompilationResult::Unsupported;
-#if !IS_MONOLITHIC
+#if WITH_HOT_RELOAD
 	bool bCanRebind = InPackages.Num() > 0;
 
 	// Verify that we're going to be able to rebind the specified packages
@@ -919,7 +919,7 @@ void FHotReloadModule::RecordAnalyticsEvent(const TCHAR* ReloadFrom, ECompilatio
 
 bool FHotReloadModule::RecompileModulesAsync( const TArray< FName > ModuleNames, const FRecompileModulesCallback& InRecompileModulesCallback, const bool bWaitForCompletion, FOutputDevice &Ar )
 {
-#if !IS_MONOLITHIC
+#if WITH_HOT_RELOAD
 	// NOTE: This method of recompiling always using a rolling file name scheme, since we never want to unload before
 	// we start recompiling, and we need the output DLL to be unlocked before we invoke the compiler
 
@@ -973,7 +973,7 @@ bool FHotReloadModule::RecompileModulesAsync( const TArray< FName > ModuleNames,
 	return bWasSuccessful;
 #else
 	return false;
-#endif // !IS_MONOLITHIC
+#endif // WITH_HOT_RELOAD
 }
 
 void FHotReloadModule::OnModuleCompileSucceeded(FName ModuleName, const FString& NewModuleFilename)
@@ -981,7 +981,7 @@ void FHotReloadModule::OnModuleCompileSucceeded(FName ModuleName, const FString&
 	// If the compile succeeded, update the module info entry with the new file name for this module
 	FModuleManager::Get().SetModuleFilename(ModuleName, NewModuleFilename);
 
-#if !IS_MONOLITHIC && WITH_EDITOR
+#if WITH_HOT_RELOAD
 	// UpdateModuleCompileData() should have been run before compiling so the
 	// data in ModuleInfo should be correct for the pre-compile dll file.
 	FModuleCompilationData& CompileData = ModuleCompileData.FindChecked(ModuleName).Get();
@@ -1007,7 +1007,7 @@ void FHotReloadModule::OnModuleCompileSucceeded(FName ModuleName, const FString&
 bool FHotReloadModule::RecompileModuleDLLs( const TArray< FModuleToRecompile >& ModuleNames, FOutputDevice& Ar )
 {
 	bool bCompileSucceeded = false;
-#if !IS_MONOLITHIC
+#if WITH_HOT_RELOAD
 	const FString AdditionalArguments = MakeUBTArgumentsForModuleCompiling();
 	if( StartCompilingModuleDLLs( FApp::GetGameName(), ModuleNames, FRecompileModulesCallback(), Ar, true, AdditionalArguments ) )
 	{
@@ -1050,7 +1050,7 @@ bool FHotReloadModule::StartCompilingModuleDLLs(const FString& GameName, const T
 	const FRecompileModulesCallback& InRecompileModulesCallback, FOutputDevice& Ar, bool bInFailIfGeneratedCodeChanges, 
 	const FString& InAdditionalCmdLineArgs )
 {
-#if PLATFORM_DESKTOP && !IS_MONOLITHIC
+#if WITH_HOT_RELOAD
 	// Keep track of what we're compiling
 	ModulesBeingCompiled = ModuleNames;
 	ModulesThatWereBeingRecompiled = ModulesBeingCompiled;
@@ -1134,7 +1134,7 @@ bool FHotReloadModule::StartCompilingModuleDLLs(const FString& GameName, const T
 
 bool FHotReloadModule::InvokeUnrealBuildToolForCompile(const FString& InCmdLineParams, FOutputDevice &Ar)
 {
-#if PLATFORM_DESKTOP && !IS_MONOLITHIC
+#if WITH_HOT_RELOAD
 
 	// Make sure we're not already compiling something!
 	check(!IsCurrentlyCompiling());
@@ -1178,12 +1178,12 @@ bool FHotReloadModule::InvokeUnrealBuildToolForCompile(const FString& InCmdLineP
 	return ProcHandle.IsValid();
 #else
 	return false;
-#endif // PLATFORM_DESKTOP && !IS_MONOLITHIC
+#endif // WITH_HOT_RELOAD
 }
 
 void FHotReloadModule::CheckForFinishedModuleDLLCompile(const bool bWaitForCompletion, bool& bCompileStillInProgress, bool& bCompileSucceeded, FOutputDevice& Ar, const FText& SlowTaskOverrideText, bool bFireEvents)
 {
-#if PLATFORM_DESKTOP && !IS_MONOLITHIC
+#if WITH_HOT_RELOAD
 	bCompileStillInProgress = false;
 	ECompilationResult::Type CompilationResult = ECompilationResult::OtherCompilationError;
 
@@ -1323,7 +1323,7 @@ void FHotReloadModule::CheckForFinishedModuleDLLCompile(const bool bWaitForCompl
 	{
 		Ar.Logf(TEXT("Error: CheckForFinishedModuleDLLCompile: There is no compilation in progress right now"));
 	}
-#endif // PLATFORM_DESKTOP && !IS_MONOLITHIC
+#endif // WITH_HOT_RELOAD
 }
 
 void FHotReloadModule::UpdateModuleCompileData(FName ModuleName)
@@ -1341,7 +1341,7 @@ void FHotReloadModule::UpdateModuleCompileData(FName ModuleName)
 	CompileData.FileTimeStamp = FDateTime(0);
 	CompileData.CompileMethod = EModuleCompileMethod::Unknown;
 
-#if !IS_MONOLITHIC && WITH_EDITOR
+#if WITH_HOT_RELOAD
 	ReadModuleCompilationInfoFromConfig(ModuleName, CompileData);
 
 	FDateTime FileTimeStamp;
diff --git a/Engine/Source/Runtime/Core/Private/HAL/ConsoleManager.cpp b/Engine/Source/Runtime/Core/Private/HAL/ConsoleManager.cpp
--- a/Engine/Source/Runtime/Core/Private/HAL/ConsoleManager.cpp
+++ b/Engine/Source/Runtime/Core/Private/HAL/ConsoleManager.cpp
@@ -909,7 +909,7 @@ IConsoleObject* FConsoleManager::AddConsoleObject(const TCHAR* Name, IConsoleObj
 		// An existing console object was found that has the same name as the object being registered.
 		// In most cases this is not allowed, but if there is a variable with the same name and is
 		// in an 'unregistered' state or we're hot reloading dlls, we may be able to replace or update that variable.
-#if !IS_MONOLITHIC
+#if WITH_HOT_RELOAD
 		const bool bCanUpdateOrReplaceObj = (ExistingObj->AsVariable()||ExistingObj->AsCommand()) && (GIsHotReload || ExistingObj->TestFlags(ECVF_Unregistered));
 #else
 		const bool bCanUpdateOrReplaceObj = ExistingObj->AsVariable() && ExistingObj->TestFlags(ECVF_Unregistered);
diff --git a/Engine/Source/Runtime/Core/Private/Modules/ModuleManager.cpp b/Engine/Source/Runtime/Core/Private/Modules/ModuleManager.cpp
--- a/Engine/Source/Runtime/Core/Private/Modules/ModuleManager.cpp
+++ b/Engine/Source/Runtime/Core/Private/Modules/ModuleManager.cpp
@@ -8,7 +8,7 @@
 
 DEFINE_LOG_CATEGORY_STATIC(LogModuleManager, Log, All);
 
-#if !IS_MONOLITHIC
+#if WITH_HOT_RELOAD
 	/** If true, we are reloading a class for HotReload */
 	CORE_API bool			GIsHotReload							= false;
 #endif
diff --git a/Engine/Source/Runtime/Core/Public/Misc/Build.h b/Engine/Source/Runtime/Core/Public/Misc/Build.h
--- a/Engine/Source/Runtime/Core/Public/Misc/Build.h
+++ b/Engine/Source/Runtime/Core/Public/Misc/Build.h
@@ -137,6 +137,12 @@
 	#error UBT should always define IS_PROGRAM to be 0 or 1
 #endif
 
+/**
+* Whether we support hot-reload. Currently requires a non-monolithic build and non-shipping configuration.
+*/
+#ifndef WITH_HOT_RELOAD
+	#define WITH_HOT_RELOAD (!IS_MONOLITHIC && !UE_BUILD_SHIPPING)
+#endif
 
 /*----------------------------------------------------------------------------
 	Optional bridge options coming from UBT, do not modify directly!
diff --git a/Engine/Source/Runtime/Core/Public/Modules/ModuleManager.h b/Engine/Source/Runtime/Core/Public/Modules/ModuleManager.h
--- a/Engine/Source/Runtime/Core/Public/Modules/ModuleManager.h
+++ b/Engine/Source/Runtime/Core/Public/Modules/ModuleManager.h
@@ -7,7 +7,7 @@
 #include "Boilerplate/ModuleBoilerplate.h"
 
 
-#if !IS_MONOLITHIC
+#if WITH_HOT_RELOAD
 	/** If true, we are reloading a class for HotReload */
 	extern CORE_API bool GIsHotReload;
 #endif
diff --git a/Engine/Source/Runtime/CoreUObject/Private/UObject/Class.cpp b/Engine/Source/Runtime/CoreUObject/Private/UObject/Class.cpp
--- a/Engine/Source/Runtime/CoreUObject/Private/UObject/Class.cpp
+++ b/Engine/Source/Runtime/CoreUObject/Private/UObject/Class.cpp
@@ -1825,7 +1825,7 @@ void UScriptStruct::DeferCppStructOps(FName Target, ICppStructOps* InCppStructOp
 {
 	if (GetDeferredCppStructOps().Contains(Target))
 	{
-#if !IS_MONOLITHIC
+#if WITH_HOT_RELOAD
 		if (!GIsHotReload) // in hot reload, we will just leak these...they may be in use.
 #endif
 		{
@@ -1859,7 +1859,7 @@ void UScriptStruct::PrepareCppStructOps()
 #if !(UE_BUILD_SHIPPING || UE_BUILD_TEST)
 		// test that the constructor is initializing everything
 		if (CppStructOps && !CppStructOps->HasZeroConstructor()
-#if !IS_MONOLITHIC
+#if WITH_HOT_RELOAD
 			&& !GIsHotReload // in hot reload, these produce bogus warnings
 #endif
 			)
@@ -3317,7 +3317,7 @@ UClass::UClass
 	*(const TCHAR**)&ClassConfigName = InConfigName;
 }
 
-#if !IS_MONOLITHIC
+#if WITH_HOT_RELOAD
 
 bool UClass::HotReloadPrivateStaticClass(
 	uint32			InSize,
@@ -3411,7 +3411,7 @@ bool UClass::HotReloadPrivateStaticClass(
 void UClass::AddNativeFunction(const ANSICHAR* InName,Native InPointer)
 {
 	FName InFName(InName);
-#if !IS_MONOLITHIC
+#if WITH_HOT_RELOAD
 	if (GIsHotReload)
 	{
 		IHotReloadInterface& HotReloadSupport = FModuleManager::LoadModuleChecked<IHotReloadInterface>("HotReload");
diff --git a/Engine/Source/Runtime/CoreUObject/Private/UObject/ScriptCore.cpp b/Engine/Source/Runtime/CoreUObject/Private/UObject/ScriptCore.cpp
--- a/Engine/Source/Runtime/CoreUObject/Private/UObject/ScriptCore.cpp
+++ b/Engine/Source/Runtime/CoreUObject/Private/UObject/ScriptCore.cpp
@@ -273,7 +273,7 @@ COREUOBJECT_API uint8 GRegisterNative( int32 NativeBytecodeIndex, const Native&
 	{
 		if( NativeBytecodeIndex<0 || (uint32)NativeBytecodeIndex>ARRAY_COUNT(GNatives) || GNatives[NativeBytecodeIndex]!=&UObject::execUndefined) 
 		{
-#if !IS_MONOLITHIC
+#if WITH_HOT_RELOAD
 			if (GIsHotReload)
 			{
 				IHotReloadInterface& HotReloadSupport = FModuleManager::LoadModuleChecked<IHotReloadInterface>("HotReload");
@@ -311,7 +311,7 @@ COREUOBJECT_API uint8 GRegisterCast( int32 CastCode, const Native& Func )
 	if (CastCode != INDEX_NONE)
 	{
 		if(  
-#if !IS_MONOLITHIC
+#if WITH_HOT_RELOAD
 			!GIsHotReload && 
 #endif
 			(CastCode<0 || (uint32)CastCode>ARRAY_COUNT(GCasts) || GCasts[CastCode]!=&UObject::execUndefined) ) 
diff --git a/Engine/Source/Runtime/CoreUObject/Private/UObject/UObjectBase.cpp b/Engine/Source/Runtime/CoreUObject/Private/UObject/UObjectBase.cpp
--- a/Engine/Source/Runtime/CoreUObject/Private/UObject/UObjectBase.cpp
+++ b/Engine/Source/Runtime/CoreUObject/Private/UObject/UObjectBase.cpp
@@ -427,7 +427,7 @@ void UObjectCompiledInDeferStruct(class UScriptStruct *(*InRegister)())
 	GetDeferredCompiledInStructRegistration().Add(InRegister);
 }
 
-#if !IS_MONOLITHIC
+#if WITH_HOT_RELOAD
 struct FStructOrEnumCompiledInfo : FFieldCompiledInInfo
 {
 	FStructOrEnumCompiledInfo(SIZE_T InClassSize, uint32 InCrc)
@@ -446,7 +446,7 @@ static TMap<FName, FStructOrEnumCompiledInfo*>& GetStructOrEnumGeneratedCodeInfo
 
 class UScriptStruct *GetStaticStruct(class UScriptStruct *(*InRegister)(), UObject* StructOuter, const TCHAR* StructName, SIZE_T Size, uint32 Crc)
 {
-#if !IS_MONOLITHIC
+#if WITH_HOT_RELOAD
 	// Try to get the existing info for this struct
 	const FName StructFName(StructName);
 	FStructOrEnumCompiledInfo* Info = nullptr;
@@ -513,7 +513,7 @@ void UObjectCompiledInDeferEnum(class UEnum *(*InRegister)())
 
 class UEnum *GetStaticEnum(class UEnum *(*InRegister)(), UObject* EnumOuter, const TCHAR* EnumName)
 {
-#if !IS_MONOLITHIC
+#if WITH_HOT_RELOAD
 	if (GIsHotReload)
 	{
 		UEnum* ReturnEnum = FindObjectChecked<UEnum>(EnumOuter, EnumName);
@@ -541,7 +541,7 @@ static TArray<FFieldCompiledInInfo*>& GetDeferredClassRegistration()
 	return DeferredClassRegistration;
 }
 
-#if !IS_MONOLITHIC
+#if WITH_HOT_RELOAD
 /** Map of deferred class registration info (including size and reflection info) */
 static TMap<FName, FFieldCompiledInInfo*>& GetDeferRegisterClassMap()
 {
@@ -560,7 +560,7 @@ static TArray<FFieldCompiledInInfo*>& GetHotReloadClasses()
 void UClassCompiledInDefer(FFieldCompiledInInfo* ClassInfo, const TCHAR* Name, SIZE_T ClassSize, uint32 Crc)
 {
 	const FName CPPClassName(Name);
-#if !IS_MONOLITHIC
+#if WITH_HOT_RELOAD
 	// Check for existing classes
 	auto ExistingClassInfo = GetDeferRegisterClassMap().Find(CPPClassName);
 	ClassInfo->bHasChanged = !ExistingClassInfo || (*ExistingClassInfo)->Size != ClassInfo->Size || (*ExistingClassInfo)->Crc != ClassInfo->Crc;
@@ -610,7 +610,7 @@ void UClassCompiledInDefer(FFieldCompiledInInfo* ClassInfo, const TCHAR* Name, S
 
 void UObjectCompiledInDefer(class UClass *(*InRegister)(), const TCHAR* Name)
 {
-#if !IS_MONOLITHIC
+#if WITH_HOT_RELOAD
 	// Either add all classes if not hot-reloading, or those which have changed
 	if (!GIsHotReload || GetDeferRegisterClassMap().FindChecked(Name)->bHasChanged)
 #endif
@@ -630,7 +630,7 @@ void UClassRegisterAllCompiledInClasses()
 	GetDeferredClassRegistration().Empty();
 }
 
-#if !IS_MONOLITHIC
+#if WITH_HOT_RELOAD
 /** Re-instance all existing classes that have changed during hot-reload */
 void UClassReplaceHotReloadClasses()
 {
@@ -737,7 +737,7 @@ void ProcessNewlyLoadedUObjects()
 		UObjectLoadAllCompiledInStructs();
 		UObjectLoadAllCompiledInDefaultProperties();		
 	}
-#if !IS_MONOLITHIC
+#if WITH_HOT_RELOAD
 	UClassReplaceHotReloadClasses();
 #endif
 }
diff --git a/Engine/Source/Runtime/CoreUObject/Public/UObject/Class.h b/Engine/Source/Runtime/CoreUObject/Public/UObject/Class.h
--- a/Engine/Source/Runtime/CoreUObject/Public/UObject/Class.h
+++ b/Engine/Source/Runtime/CoreUObject/Public/UObject/Class.h
@@ -1823,7 +1823,7 @@ public:
 		const TCHAR* InClassConfigName, EObjectFlags InFlags, void(*InClassConstructor)(const class FPostConstructInitializeProperties&),
 		void(*InClassAddReferencedObjects)(UObject*, class FReferenceCollector&));
 
-#if !IS_MONOLITHIC
+#if WITH_HOT_RELOAD
 	/**
 	 * Called when a class is reloading from a DLL...updates various information in-place.
 	 * @param	InSize							sizeof the class
@@ -2190,7 +2190,7 @@ COREUOBJECT_API void InitializePrivateStaticClass(
 template<class TClass>
 void GetPrivateStaticClassBody( const TCHAR* PackageName, const TCHAR* Name, UClass*& ReturnClass, void (*RegisterNativeFunc)() )
 { 
-#if !IS_MONOLITHIC
+#if WITH_HOT_RELOAD
 	if (GIsHotReload)
 	{
 		UPackage* Package = FindPackage(NULL, PackageName);
diff --git a/Engine/Source/Runtime/CoreUObject/Public/UObject/CoreObject.h b/Engine/Source/Runtime/CoreUObject/Public/UObject/CoreObject.h
--- a/Engine/Source/Runtime/CoreUObject/Public/UObject/CoreObject.h
+++ b/Engine/Source/Runtime/CoreUObject/Public/UObject/CoreObject.h
@@ -47,11 +47,6 @@ public:
 	static FOnPackageMarkedDirty PackageMarkedDirtyEvent;
 
 private:
-#if !IS_MONOLITHIC
-	/** True if we have attempted to load code for this package from a DLL already */
-	bool bHasAttemptedBind;
-#endif
-
 	/** Used by the editor to determine if a package has been changed.																							*/
 	bool	bDirty;
 public:
diff --git a/Engine/Source/Runtime/Engine/Classes/Engine/Engine.h b/Engine/Source/Runtime/Engine/Classes/Engine/Engine.h
--- a/Engine/Source/Runtime/Engine/Classes/Engine/Engine.h
+++ b/Engine/Source/Runtime/Engine/Classes/Engine/Engine.h
@@ -1632,7 +1632,7 @@ public:
 
 	// Compile in Debug or Development
 #if !(UE_BUILD_SHIPPING || UE_BUILD_TEST)
-#if !IS_MONOLITHIC
+#if WITH_HOT_RELOAD
 	bool HandleHotReloadCommand( const TCHAR* Cmd, FOutputDevice& Ar );
 #endif
 	bool HandleDumpConsoleCommandsCommand( const TCHAR* Cmd, FOutputDevice& Ar, UWorld* InWorld );
diff --git a/Engine/Source/Runtime/Engine/Private/UnrealEngine.cpp b/Engine/Source/Runtime/Engine/Private/UnrealEngine.cpp
--- a/Engine/Source/Runtime/Engine/Private/UnrealEngine.cpp
+++ b/Engine/Source/Runtime/Engine/Private/UnrealEngine.cpp
@@ -2217,12 +2217,12 @@ bool UEngine::Exec( UWorld* InWorld, const TCHAR* Cmd, FOutputDevice& Ar )
 #endif
 
 #if !(UE_BUILD_SHIPPING || UE_BUILD_TEST)
-#if !IS_MONOLITHIC
+#if WITH_HOT_RELOAD
 	else if( FParse::Command(&Cmd,TEXT("HotReload")) )
 	{
 		return HandleHotReloadCommand( Cmd, Ar );
 	}
-#endif // !IS_MONOLITHIC
+#endif // WITH_HOT_RELOAD
 	else if (FParse::Command(&Cmd, TEXT("DumpConsoleCommands")))
 	{
 		return HandleDumpConsoleCommandsCommand( Cmd, Ar, InWorld );
@@ -2651,7 +2651,7 @@ bool UEngine::HandleMergeMeshCommand( const TCHAR* Cmd, FOutputDevice& Ar, UWorl
 
 
 #if !(UE_BUILD_SHIPPING || UE_BUILD_TEST)
-#if !IS_MONOLITHIC
+#if WITH_HOT_RELOAD
 bool UEngine::HandleHotReloadCommand( const TCHAR* Cmd, FOutputDevice& Ar )
 {
 	FString Module(FParse::Token(Cmd, 0));
@@ -2672,7 +2672,7 @@ bool UEngine::HandleHotReloadCommand( const TCHAR* Cmd, FOutputDevice& Ar )
 	}
 	return true;
 }
-#endif // !IS_MONOLITHIC
+#endif // WITH_HOT_RELOAD
 
 bool UEngine::HandleDumpConsoleCommandsCommand( const TCHAR* Cmd, FOutputDevice& Ar, UWorld* InWorld )
 {
