Added WITH_ENGINE defines around HotReload code specifically dependent on engine types, so programs not dependent on engine can use the HotReload module.

diff --git a/Engine/Source/Developer/HotReload/HotReload.Build.cs b/Engine/Source/Developer/HotReload/HotReload.Build.cs
--- a/Engine/Source/Developer/HotReload/HotReload.Build.cs
+++ b/Engine/Source/Developer/HotReload/HotReload.Build.cs
@@ -25,13 +25,15 @@ public class HotReload : ModuleRules
 				}
 				);
 
-
-			PrivateDependencyModuleNames.AddRange(
-				new string[] 
-				{
-					"Engine",
-					"UnrealEd", 
-				}
-			);
+            if (UEBuildConfiguration.bCompileAgainstEngine)
+            {
+                PrivateDependencyModuleNames.AddRange(
+				    new string[] 
+				    {
+					    "Engine",
+					    "UnrealEd", 
+				    }
+			    );
+            }
 	}
 }
\ No newline at end of file
diff --git a/Engine/Source/Developer/HotReload/Private/HotReload.cpp b/Engine/Source/Developer/HotReload/Private/HotReload.cpp
--- a/Engine/Source/Developer/HotReload/Private/HotReload.cpp
+++ b/Engine/Source/Developer/HotReload/Private/HotReload.cpp
@@ -1,9 +1,13 @@
 // Copyright 1998-2014 Epic Games, Inc. All Rights Reserved.
 
 #include "HotReloadPrivatePCH.h"
+
+#if WITH_ENGINE
 #include "Editor/UnrealEd/Public/Kismet2/KismetReinstanceUtilities.h"
-#include "Runtime/Analytics/Analytics/Public/Interfaces/IAnalyticsProvider.h"
 #include "EngineAnalytics.h"
+#endif
+
+#include "Runtime/Analytics/Analytics/Public/Interfaces/IAnalyticsProvider.h"
 #include "ScopedTimers.h"
 #include "DesktopPlatformModule.h"
 
@@ -131,10 +135,12 @@ private:
 	 */
 	void GetPackagesToRebindAndDependentModules(const TArray<FString>& InGameModuleNames, TArray<UPackage*>& OutPackagesToRebind, TArray<FName>& OutDependentModules);
 
+#if WITH_ENGINE
 	/**
 	 * Called from CoreUObject to re-instance hot-reloaded classes
 	 */
 	void ReinstanceClass(UClass* OldClass, UClass* NewClass);
+#endif
 
 	/**
 	 * Tick function for FTicker: checks for re-loaded modules and does hot-reload from IDE
@@ -307,9 +313,11 @@ void FHotReloadModule::StartupModule()
 {
 	bIsHotReloadingFromEditor = false;
 
+#if WITH_ENGINE
 	// Register re-instancing delegate (Core)
 	FCoreUObjectDelegates::ReplaceHotReloadClassDelegate.BindRaw(this, &FHotReloadModule::ReinstanceClass);
-	
+#endif
+
 	// Register directory watcher delegate
 	InitHotReloadWatcher();
 
@@ -739,12 +747,14 @@ ECompilationResult::Type FHotReloadModule::RebindPackagesInternal(TArray<UPackag
 	return Result;
 }
 
+#if WITH_ENGINE
 void FHotReloadModule::ReinstanceClass(UClass* OldClass, UClass* NewClass)
 {
 	UE_LOG(LogHotReload, Log, TEXT("Re-instancing %s after hot-reload."), *NewClass->GetName());
 	FBlueprintCompileReinstancer ReinstanceHelper(NewClass, OldClass);
 	ReinstanceHelper.ReinstanceObjects();
 }
+#endif
 
 void FHotReloadModule::GetGameModules(TArray<FString>& OutGameModules)
 {
@@ -905,6 +915,7 @@ void FHotReloadModule::DoHotReloadFromIDE()
 
 void FHotReloadModule::RecordAnalyticsEvent(const TCHAR* ReloadFrom, ECompilationResult::Type Result, double Duration, int32 PackageCount, int32 DependentModulesCount)
 {
+#if WITH_ENGINE
 	if (FEngineAnalytics::IsAvailable())
 	{
 		TArray< FAnalyticsEventAttribute > ReloadAttribs;
@@ -915,6 +926,7 @@ void FHotReloadModule::RecordAnalyticsEvent(const TCHAR* ReloadFrom, ECompilatio
 		ReloadAttribs.Add(FAnalyticsEventAttribute(TEXT("DependentModules"), FString::Printf(TEXT("%d"), DependentModulesCount)));
 		FEngineAnalytics::GetProvider().RecordEvent(TEXT("Editor.Usage.HotReload"), ReloadAttribs);
 	}
+#endif
 }
 
 bool FHotReloadModule::RecompileModulesAsync( const TArray< FName > ModuleNames, const FRecompileModulesCallback& InRecompileModulesCallback, const bool bWaitForCompletion, FOutputDevice &Ar )
