When compiling without the editor don't assume the TargetName is UE4Editor.

diff --git a/Engine/Source/Developer/HotReload/Private/HotReload.cpp b/Engine/Source/Developer/HotReload/Private/HotReload.cpp
--- a/Engine/Source/Developer/HotReload/Private/HotReload.cpp
+++ b/Engine/Source/Developer/HotReload/Private/HotReload.cpp
@@ -1108,17 +1108,17 @@ bool FHotReloadModule::StartCompilingModuleDLLs(const FString& GameName, const T
 	// Shared PCH does no work with hot-reloading modules as we don't scan all modules for them.
 	ExtraArg += TEXT("-nosharedpch ");
 
-	// If the're no game modules loaded, then it's not a code-based project and the target
+	FString TargetName = GameName;
+
+#if WITH_EDITOR
+	// If there are no game modules loaded, then it's not a code-based project and the target
 	// for UBT should be the editor.
-	FString TargetName;
-	if (IsAnyGameModuleLoaded())
-	{
-		TargetName = GameName;
-	}
-	else
+	if (!IsAnyGameModuleLoaded())
 	{
 		TargetName = TEXT("UE4Editor");
 	}
+#endif
+
 	FString CmdLineParams = FString::Printf( TEXT( "%s%s %s %s %s%s" ), 
 		*TargetName, *ModuleArg, 
 		BuildPlatformName, BuildConfigurationName, 
