Hot-reload: all CDOs will now get re-constructed after hot-reload. If a CDO has changed, the class will be re-instanced.

#change Moved hot-reload-only code from FBlueprintCompileReinstancer to FHotReloadClassReinstancer
#change It's not possible to derive classes from FBlueprintCompileReinstancer
#change Added code to handle re-creating CDOs after hot-reload even if the class hasn't changed.

diff --git a/Engine/Source/Developer/HotReload/Private/HotReload.cpp b/Engine/Source/Developer/HotReload/Private/HotReload.cpp
--- a/Engine/Source/Developer/HotReload/Private/HotReload.cpp
+++ b/Engine/Source/Developer/HotReload/Private/HotReload.cpp
@@ -1,15 +1,13 @@
 // Copyright 1998-2014 Epic Games, Inc. All Rights Reserved.
 
 #include "HotReloadPrivatePCH.h"
-
-#if WITH_ENGINE
-#include "Editor/UnrealEd/Public/Kismet2/KismetReinstanceUtilities.h"
-#include "EngineAnalytics.h"
-#endif
-
 #include "Runtime/Analytics/Analytics/Public/Interfaces/IAnalyticsProvider.h"
 #include "ScopedTimers.h"
 #include "DesktopPlatformModule.h"
+#if WITH_ENGINE
+#include "HotReloadClassReinstancer.h"
+#include "EngineAnalytics.h"
+#endif
 
 DEFINE_LOG_CATEGORY(LogHotReload);
 
@@ -749,10 +747,13 @@ ECompilationResult::Type FHotReloadModule::RebindPackagesInternal(TArray<UPackag
 
 #if WITH_ENGINE
 void FHotReloadModule::ReinstanceClass(UClass* OldClass, UClass* NewClass)
-{
-	UE_LOG(LogHotReload, Log, TEXT("Re-instancing %s after hot-reload."), *NewClass->GetName());
-	FBlueprintCompileReinstancer ReinstanceHelper(NewClass, OldClass);
-	ReinstanceHelper.ReinstanceObjects();
+{	
+	FHotReloadClassReinstancer ReinstanceHelper(NewClass, OldClass);
+	if (ReinstanceHelper.ClassNeedsReinstancing())
+	{
+		UE_LOG(LogHotReload, Log, TEXT("Re-instancing %s after hot-reload."), NewClass ? *NewClass->GetName() : *OldClass->GetName());
+		ReinstanceHelper.ReinstanceObjects();
+	}
 }
 #endif
 
diff --git a/Engine/Source/Developer/HotReload/Private/HotReloadClassReinstancer.cpp b/Engine/Source/Developer/HotReload/Private/HotReloadClassReinstancer.cpp
new file mode 100644
--- /dev/null
+++ b/Engine/Source/Developer/HotReload/Private/HotReloadClassReinstancer.cpp
@@ -0,0 +1,235 @@
+// Copyright 1998-2014 Epic Games, Inc. All Rights Reserved.
+
+#include "HotReloadPrivatePCH.h"
+#include "HotReloadClassReinstancer.h"
+
+#if WITH_ENGINE
+
+void FHotReloadClassReinstancer::SetupNewClassReinstancing(UClass* InNewClass, UClass* InOldClass)
+{
+	// Set base class members to valid values
+	ClassToReinstance = InNewClass;
+	DuplicatedClass = InOldClass;
+	OriginalCDO = InOldClass->GetDefaultObject();
+	bHasReinstanced = false;
+	bSkipGarbageCollection = false;
+	bNeedsReinstancing = true;
+
+	SaveClassFieldMapping(InOldClass);
+
+	TArray<UClass*> ChildrenOfClass;
+	GetDerivedClasses(InOldClass, ChildrenOfClass);
+	for (auto ClassIt = ChildrenOfClass.CreateConstIterator(); ClassIt; ++ClassIt)
+	{
+		UClass* ChildClass = *ClassIt;
+		UBlueprint* ChildBP = Cast<UBlueprint>(ChildClass->ClassGeneratedBy);
+		if (ChildBP && !ChildBP->HasAnyFlags(RF_BeingRegenerated))
+		{
+			// If this is a direct child, change the parent and relink so the property chain is valid for reinstancing
+			if (!ChildBP->HasAnyFlags(RF_NeedLoad))
+			{
+				if (ChildClass->GetSuperClass() == InOldClass)
+				{
+					ReparentChild(ChildBP);
+				}
+
+				Children.AddUnique(ChildBP);
+			}
+			else
+			{
+				// If this is a child that caused the load of their parent, relink to the REINST class so that we can still serialize in the CDO, but do not add to later processing
+				ReparentChild(ChildClass);
+			}
+		}
+	}
+
+	// Finally, remove the old class from Root so that it can get GC'd and mark it as CLASS_NewerVersionExists
+	InOldClass->RemoveFromRoot();
+	InOldClass->ClassFlags |= CLASS_NewerVersionExists;
+}
+
+void FHotReloadClassReinstancer::SerializeCDOProperties(UObject* InObject, TArray<uint8>& OutData)
+{
+	// Creates a mem-comparable CDO data
+	class FCDOWriter : public FMemoryWriter
+	{
+		/** Objects already visited by this archive */
+		TSet<UObject*>& VisitedObjects;
+
+	public:
+		/** Serializes all script properties of the provided DefaultObject */
+		FCDOWriter(TArray<uint8>& InBytes, UObject* DefaultObject, TSet<UObject*>& InVisitedObjects)
+			: FMemoryWriter(InBytes, /* bIsPersistent = */ false, /* bSetOffset = */ true)
+			, VisitedObjects(InVisitedObjects)
+		{
+			DefaultObject->SerializeScriptProperties(*this);
+		}
+		/** Serializes an object. Only name and class for normal references, deep serialization for DSOs */
+		virtual FArchive& operator<<(class UObject*& InObj) override
+		{
+			FArchive& Ar = *this;
+			if (InObj)
+			{
+				FName ClassName = InObj->GetClass()->GetFName();
+				FName ObjectName = InObj->GetFName();
+				Ar << ClassName;
+				Ar << ObjectName;
+				if (GSerializedProperty && GSerializedProperty->ContainsInstancedObjectProperty() && !VisitedObjects.Contains(InObj))
+				{
+					// Serialize all DSO properties too
+					VisitedObjects.Add(InObj);
+					FCDOWriter DefaultSubobjectWriter(Bytes, InObj, VisitedObjects);
+				}
+			}
+			else
+			{
+				FName UnusedName = NAME_None;
+				Ar << UnusedName;
+				Ar << UnusedName;
+			}
+
+			return *this;
+		}
+		/** Serializes an FNane as an index and number */
+		virtual FArchive& operator<<(FName& InName) override
+		{
+			FArchive& Ar = *this;
+			int32 Index = InName.GetIndex();
+			int32 Number = InName.GetNumber();
+			Ar << Index;
+			Ar << Number;
+			return Ar;
+		}
+		virtual FArchive& operator<<(FLazyObjectPtr& LazyObjectPtr) override
+		{
+			FArchive& Ar = *this;
+			UObject* Obj = LazyObjectPtr.Get();
+			Ar << Obj;
+			return *this;
+		}
+		virtual FArchive& operator<<(FAssetPtr& AssetPtr) override
+		{
+			FArchive& Ar = *this;
+			UObject* Obj = AssetPtr.Get();
+			Ar << Obj;
+			return Ar;
+		}
+		virtual FArchive& operator<<(FStringAssetReference& Value) override
+		{
+			FArchive& Ar = *this;
+			Value.Serialize(Ar);
+			return Ar;
+		}
+		/** Archive name, for debugging */
+		virtual FString GetArchiveName() const override { return TEXT("FCDOWriter"); }
+	};
+	TSet<UObject*> VisitedObjects;
+	FCDOWriter Ar(OutData, InObject, VisitedObjects);
+}
+
+void FHotReloadClassReinstancer::ReconstructClassDefaultObject(UClass* InOldClass)
+{
+	// Remember all the basic info about the object before we destroy it
+	UObject* OldCDO = InOldClass->GetDefaultObject();
+	EObjectFlags CDOFlags = OldCDO->GetFlags();
+	UObject* CDOOuter = OldCDO->GetOuter();
+	FName CDOName = OldCDO->GetFName();
+
+	// Get the parent CDO
+	UClass* ParentClass = InOldClass->GetSuperClass();
+	UObject* ParentDefaultObject = NULL;
+	if (ParentClass != NULL)
+	{
+		ParentDefaultObject = ParentClass->GetDefaultObject(); // Force the default object to be constructed if it isn't already
+	}
+
+	if (!OldCDO->HasAnyFlags(RF_FinishDestroyed))
+	{
+		// Begin the asynchronous object cleanup.
+		OldCDO->ConditionalBeginDestroy();
+
+		// Wait for the object's asynchronous cleanup to finish.
+		while (!OldCDO->IsReadyForFinishDestroy())
+		{
+			FPlatformProcess::Sleep(0);
+		}
+		// Finish destroying the object.
+		OldCDO->ConditionalFinishDestroy();
+	}
+	OldCDO->~UObject();
+
+	// Re-create
+	FMemory::Memzero((void*)OldCDO, InOldClass->GetPropertiesSize());
+	new ((void *)OldCDO) UObjectBase(InOldClass, CDOFlags, CDOOuter, CDOName);
+	const bool bShouldInitilizeProperties = false;
+	const bool bCopyTransientsFromClassDefaults = false;
+	(*InOldClass->ClassConstructor)(FPostConstructInitializeProperties(OldCDO, ParentDefaultObject, bCopyTransientsFromClassDefaults, bShouldInitilizeProperties));
+}
+
+void FHotReloadClassReinstancer::RecreateCDOAndSetupOldClassReinstancing(UClass* InOldClass)
+{
+	// Set base class members to valid values
+	ClassToReinstance = InOldClass;
+	DuplicatedClass = InOldClass;
+	OriginalCDO = InOldClass->GetDefaultObject();
+	bHasReinstanced = false;
+	bSkipGarbageCollection = false;
+	bNeedsReinstancing = false;
+
+	// Collect the original property values
+	TArray<uint8> OriginalCDOProperties;
+	SerializeCDOProperties(InOldClass->GetDefaultObject(), OriginalCDOProperties);
+	
+	// Destroy and re-create the CDO, re-running its constructor
+	ReconstructClassDefaultObject(InOldClass);
+
+	// Collect the property values after re-constructing the CDO
+	TArray<uint8> ReconstructedCDOProperties;
+	SerializeCDOProperties(InOldClass->GetDefaultObject(), ReconstructedCDOProperties);
+
+	// We only want to re-instance the old class if its CDO's values have changed or any of its DSOs' property values have changed
+	if (OriginalCDOProperties.Num() != ReconstructedCDOProperties.Num() ||
+		FMemory::Memcmp(OriginalCDOProperties.GetTypedData(), ReconstructedCDOProperties.GetTypedData(), OriginalCDOProperties.Num()))
+	{
+		bNeedsReinstancing = true;
+		SaveClassFieldMapping(InOldClass);
+
+		TArray<UClass*> ChildrenOfClass;
+		GetDerivedClasses(InOldClass, ChildrenOfClass);
+		for (auto ClassIt = ChildrenOfClass.CreateConstIterator(); ClassIt; ++ClassIt)
+		{
+			UClass* ChildClass = *ClassIt;
+			UBlueprint* ChildBP = Cast<UBlueprint>(ChildClass->ClassGeneratedBy);
+			if (ChildBP && !ChildBP->HasAnyFlags(RF_BeingRegenerated))
+			{
+				if (!ChildBP->HasAnyFlags(RF_NeedLoad))
+				{
+					Children.AddUnique(ChildBP);
+				}
+			}
+		}
+	}
+}
+
+FHotReloadClassReinstancer::FHotReloadClassReinstancer(UClass* InNewClass, UClass* InOldClass)
+{
+	// If InNewClass is NULL, then the old class has not changed after hot-reload.
+	// However, we still need to check for changes to its constructor code (CDO values).
+	if (InNewClass)
+	{
+		SetupNewClassReinstancing(InNewClass, InOldClass);
+	}
+	else
+	{
+		RecreateCDOAndSetupOldClassReinstancing(InOldClass);
+	}
+}
+
+FHotReloadClassReinstancer::~FHotReloadClassReinstancer()
+{
+	// Make sure the base class does not remove the DuplicatedClass from root, we not always want it.
+	// For example when we're just reconstructing CDOs. Other cases are handled by HotReloadClassReinstancer.
+	DuplicatedClass = NULL;
+}
+
+#endif
\ No newline at end of file
diff --git a/Engine/Source/Developer/HotReload/Private/HotReloadClassReinstancer.h b/Engine/Source/Developer/HotReload/Private/HotReloadClassReinstancer.h
new file mode 100644
--- /dev/null
+++ b/Engine/Source/Developer/HotReload/Private/HotReloadClassReinstancer.h
@@ -0,0 +1,61 @@
+// Copyright 1998-2014 Epic Games, Inc. All Rights Reserved.
+#pragma once
+
+#if WITH_ENGINE
+
+#include "Editor/UnrealEd/Public/Kismet2/KismetReinstanceUtilities.h"
+
+/**
+ * Helper class used for re-instancing native and blueprint classes after hot-reload
+ */
+class FHotReloadClassReinstancer : public FBlueprintCompileReinstancer
+{
+	/** True if the provided native class needs re-instancing */
+	bool bNeedsReinstancing;
+
+	/** 
+	 * Sets the re-instancer up for new class re-instancing 
+	 *
+	 * @param InNewClass Class that has changed after hot-reload
+	 * @param InOldClass Class before it was hot-reloaded
+	 */
+	void SetupNewClassReinstancing(UClass* InNewClass, UClass* InOldClass);
+
+	/**
+	* Sets the re-instancer up for old class re-instancing. Always re-creates the CDO.
+	*
+	* @param InOldClass Class that has NOT changed after hot-reload
+	*/
+	void RecreateCDOAndSetupOldClassReinstancing(UClass* InOldClass);
+	
+	/**
+	* Creates a mem-comparable array of data containing CDO property values.
+	*
+	* @param InObject CDO
+	* @param OutData Data containing all of the CDO property values
+	*/
+	void SerializeCDOProperties(UObject* InObject, TArray<uint8>& OutData);
+
+	/**
+	* Re-creates class default object
+	*
+	* @param InOldClass Class that has NOT changed after hot-reload
+	*/
+	void ReconstructClassDefaultObject(UClass* InOldClass);
+
+public:
+
+	/** Sets the re-instancer up to re-instance native classes */
+	FHotReloadClassReinstancer(UClass* InNewClass, UClass* InOldClass);
+	
+	/** Destructor */
+	virtual ~FHotReloadClassReinstancer();
+
+	/** If true, the class needs re-instancing */
+	FORCEINLINE bool ClassNeedsReinstancing() const
+	{
+		return bNeedsReinstancing;
+	}
+};
+
+#endif // WITH_ENGINE
\ No newline at end of file
diff --git a/Engine/Source/Editor/UnrealEd/Private/Kismet2/KismetReinstanceUtilities.cpp b/Engine/Source/Editor/UnrealEd/Private/Kismet2/KismetReinstanceUtilities.cpp
--- a/Engine/Source/Editor/UnrealEd/Private/Kismet2/KismetReinstanceUtilities.cpp
+++ b/Engine/Source/Editor/UnrealEd/Private/Kismet2/KismetReinstanceUtilities.cpp
@@ -10,46 +10,6 @@
 /////////////////////////////////////////////////////////////////////////////////
 // FBlueprintCompileReinstancer
 
-FBlueprintCompileReinstancer::FBlueprintCompileReinstancer(UClass* InNewClass, UClass* InOldClass)
-: ClassToReinstance(InNewClass)
-, DuplicatedClass(InOldClass)
-, OriginalCDO(InOldClass->GetDefaultObject())
-, bHasReinstanced(false)
-, bSkipGarbageCollection(false)
-{
-	SaveClassFieldMapping(InOldClass);
-
-	TArray<UClass*> ChildrenOfClass;
-	GetDerivedClasses(InOldClass, ChildrenOfClass);
-	for (auto ClassIt = ChildrenOfClass.CreateConstIterator(); ClassIt; ++ClassIt)
-	{
-		UClass* ChildClass = *ClassIt;
-		UBlueprint* ChildBP = Cast<UBlueprint>(ChildClass->ClassGeneratedBy);
-		if (ChildBP && !ChildBP->HasAnyFlags(RF_BeingRegenerated))
-		{
-			// If this is a direct child, change the parent and relink so the property chain is valid for reinstancing
-			if (!ChildBP->HasAnyFlags(RF_NeedLoad))
-			{
-				if (ChildClass->GetSuperClass() == InOldClass)
-				{
-					ReparentChild(ChildBP);
-				}
-
-				Children.AddUnique(ChildBP);
-			}
-			else
-			{
-				// If this is a child that caused the load of their parent, relink to the REINST class so that we can still serialize in the CDO, but do not add to later processing
-				ReparentChild(ChildClass);
-			}
-		}
-	}
-
-	// Finally, remove the old class from Root so that it can get GC'd and mark it as CLASS_NewerVersionExists
-	InOldClass->RemoveFromRoot();
-	InOldClass->ClassFlags |= CLASS_NewerVersionExists;
-}
-
 FBlueprintCompileReinstancer::FBlueprintCompileReinstancer(UClass* InClassToReinstance, bool bIsBytecodeOnly, bool bSkipGC)
 	: ClassToReinstance(InClassToReinstance)
 	, DuplicatedClass(NULL)
diff --git a/Engine/Source/Editor/UnrealEd/Public/Kismet2/KismetReinstanceUtilities.h b/Engine/Source/Editor/UnrealEd/Public/Kismet2/KismetReinstanceUtilities.h
--- a/Engine/Source/Editor/UnrealEd/Public/Kismet2/KismetReinstanceUtilities.h
+++ b/Engine/Source/Editor/UnrealEd/Public/Kismet2/KismetReinstanceUtilities.h
@@ -10,7 +10,8 @@
 
 class UNREALED_API FBlueprintCompileReinstancer
 {
-private:
+protected:
+
 	/** Reference to the class we're actively reinstancing */
 	UClass* ClassToReinstance;
 
@@ -37,12 +38,10 @@ private:
 	bool bSkipGarbageCollection;
 
 public:
-	~FBlueprintCompileReinstancer();
+	virtual ~FBlueprintCompileReinstancer();
 
 	/** Sets the reinstancer up to work on every object of the specified class */
 	FBlueprintCompileReinstancer(UClass* InClassToReinstance, bool bIsBytecodeOnly = false, bool bSkipGC = false);
-	/** Sets the reinstancer up to reinstance native classes */
-	FBlueprintCompileReinstancer(UClass* InNewClass, UClass* InOldClass);
 
 	/** Saves a mapping of field names to their UField equivalents, so we can remap any bytecode that references them later */
 	void SaveClassFieldMapping(UClass* InClassToReinstance);
@@ -62,9 +61,16 @@ public:
 	/** Verify that all instances of the duplicated class have been replaced and collected */
 	void VerifyReplacement();
 
-private:
-	/** Hide the default constructor for the reinstancer */
-	FBlueprintCompileReinstancer();
+protected:
+
+	/** Default constructor, can only be used by derived classes */
+	FBlueprintCompileReinstancer()
+		: ClassToReinstance(NULL)
+		, DuplicatedClass(NULL)
+		, OriginalCDO(NULL)
+		, bHasReinstanced(false)
+		, bSkipGarbageCollection(false)
+	{}
 
 	/** Reparents the specified blueprint or class to be the duplicated class in order to allow properties to be copied from the previous CDO to the new one */
 	void ReparentChild(UBlueprint* ChildBP);
diff --git a/Engine/Source/Runtime/CoreUObject/Private/UObject/UObjectBase.cpp b/Engine/Source/Runtime/CoreUObject/Private/UObject/UObjectBase.cpp
--- a/Engine/Source/Runtime/CoreUObject/Private/UObject/UObjectBase.cpp
+++ b/Engine/Source/Runtime/CoreUObject/Private/UObject/UObjectBase.cpp
@@ -568,21 +568,21 @@ void UClassCompiledInDefer(FFieldCompiledInInfo* ClassInfo, const TCHAR* Name, S
 	{
 		// Class exists, this can only happen during hot-reload
 		check(GIsHotReload);
+
+		// Get the native name
+		FString NameWithoutPrefix(Name);
+		NameWithoutPrefix = NameWithoutPrefix.Mid(1);
+		auto ExistingClass = FindObjectChecked<UClass>(ANY_PACKAGE, *NameWithoutPrefix);
+
 		if (ClassInfo->bHasChanged)
 		{
-			// Get the native name
-			FString NameWithoutPrefix(Name);
-			NameWithoutPrefix = NameWithoutPrefix.Mid(1);
-
 			// Rename the old class and move it to transient package
-			auto ExistingClass = FindObjectChecked<UClass>(ANY_PACKAGE, *NameWithoutPrefix);
 			ExistingClass->ClearFlags(RF_RootSet | RF_Standalone | RF_Public);
 			ExistingClass->GetDefaultObject()->ClearFlags(RF_RootSet | RF_Standalone | RF_Public);
 			const FName OldClassRename = MakeUniqueObjectName(GetTransientPackage(), ExistingClass->GetClass(), *FString::Printf(TEXT("HOTRELOADED_%s"), *NameWithoutPrefix));
 			ExistingClass->Rename(*OldClassRename.ToString(), GetTransientPackage());
 			ExistingClass->SetFlags(RF_Transient);
 			ExistingClass->AddToRoot();
-			ClassInfo->OldField = ExistingClass;
 
 			// Make sure enums de-register their names BEFORE we create the new class, otherwise there will be name conflicts
 			TArray<UObject*> ClassSubobjects;
@@ -594,9 +594,10 @@ void UClassCompiledInDefer(FFieldCompiledInInfo* ClassInfo, const TCHAR* Name, S
 					Enum->RemoveNamesFromMasterList();
 				}
 			}
-
-			GetHotReloadClasses().Add(ClassInfo);
 		}
+		ClassInfo->OldField = ExistingClass;
+		GetHotReloadClasses().Add(ClassInfo);
+
 		*ExistingClassInfo = ClassInfo;
 	}
 	else
@@ -638,7 +639,14 @@ void UClassReplaceHotReloadClasses()
 	{	
 		for (auto Class : GetHotReloadClasses())
 		{
-			FCoreUObjectDelegates::ReplaceHotReloadClassDelegate.Execute(CastChecked<UClass>(Class->OldField), CastChecked<UClass>(Class->Register()));
+			if (Class->bHasChanged)
+			{
+				FCoreUObjectDelegates::ReplaceHotReloadClassDelegate.Execute(CastChecked<UClass>(Class->OldField), CastChecked<UClass>(Class->Register()));
+			}
+			else
+			{
+				FCoreUObjectDelegates::ReplaceHotReloadClassDelegate.Execute(CastChecked<UClass>(Class->OldField), NULL);
+			}
 		}
 	}
 	GetHotReloadClasses().Empty();
diff --git a/Engine/Source/Runtime/CoreUObject/Public/UObject/UObjectGlobals.h b/Engine/Source/Runtime/CoreUObject/Public/UObject/UObjectGlobals.h
--- a/Engine/Source/Runtime/CoreUObject/Public/UObject/UObjectGlobals.h
+++ b/Engine/Source/Runtime/CoreUObject/Public/UObject/UObjectGlobals.h
@@ -1535,8 +1535,8 @@ struct COREUOBJECT_API FCoreUObjectDelegates
 	/** Delegate used by SavePackage() to check whether a package should be saved */
 	static FIsPackageOKToSaveDelegate IsPackageOKToSaveDelegate;
 
+	/** Delegate for replacing hot-reloaded classes that changed after hot-reload */
 	DECLARE_DELEGATE_TwoParams(FReplaceHotReloadClassDelegate, UClass*, UClass*);
-
 	static FReplaceHotReloadClassDelegate ReplaceHotReloadClassDelegate;
 };
 
