Fixed build error in FCDOWriter

FName no longer has GetIndex, it has GetComparisonIndex and GetDisplayIndex instead.

diff --git a/Engine/Source/Developer/HotReload/Private/HotReloadClassReinstancer.cpp b/Engine/Source/Developer/HotReload/Private/HotReloadClassReinstancer.cpp
--- a/Engine/Source/Developer/HotReload/Private/HotReloadClassReinstancer.cpp
+++ b/Engine/Source/Developer/HotReload/Private/HotReloadClassReinstancer.cpp
@@ -90,13 +90,15 @@ void FHotReloadClassReinstancer::SerializeCDOProperties(UObject* InObject, TArra
 
 			return *this;
 		}
-		/** Serializes an FNane as an index and number */
+		/** Serializes an FName as its index and number */
 		virtual FArchive& operator<<(FName& InName) override
 		{
 			FArchive& Ar = *this;
-			int32 Index = InName.GetIndex();
+			NAME_INDEX ComparisonIndex = InName.GetComparisonIndex();
+			NAME_INDEX DisplayIndex = InName.GetDisplayIndex();
 			int32 Number = InName.GetNumber();
-			Ar << Index;
+			Ar << ComparisonIndex;
+			Ar << DisplayIndex;
 			Ar << Number;
 			return Ar;
 		}
