Hot-reload: performing hot-reload when there's been no code changes will no longer result in producing new DLLs.

#ttp 345522: HotReload: Should not try to link DLLs when target is up to date

#change Added code to handle up-to-date state of modules in hot-reload code.
#change Added ECompilationResult::UpToDate and ECompilationResult::Canceled, ECompilationResult::Failed()
#change UBT: added -canskiplink command line param and support for skipping link actions when there was nothing to compile
#change extended the duration of re-compile notifications (TTP# 346604 NUXF-246 Hot compile fail message fades)
#change Added 'Compile Canceled' notification (instead of 'Compile Failed')

diff --git a/Engine/Source/Developer/HotReload/Private/HotReload.cpp b/Engine/Source/Developer/HotReload/Private/HotReload.cpp
--- a/Engine/Source/Developer/HotReload/Private/HotReload.cpp
+++ b/Engine/Source/Developer/HotReload/Private/HotReload.cpp
@@ -116,12 +116,12 @@ private:
 	/**
 	 * Does the actual hot-reload, unloads old modules, loads new ones
 	 */
-	ECompilationResult::Type DoHotReloadInternal(bool bRecompileFinished, bool bRecompileSucceeded, TArray<UPackage*> Packages, TArray<FName> InDependentModules, FOutputDevice &HotReloadAr);
+	ECompilationResult::Type DoHotReloadInternal(bool bRecompileFinished, ECompilationResult::Type CompilationResult, TArray<UPackage*> Packages, TArray<FName> InDependentModules, FOutputDevice &HotReloadAr);
 
 	/**
 	* Callback for async ompilation
 	*/
-	void DoHotReloadCallback(bool bRecompileFinished, bool bRecompileSucceeded, TArray<UPackage*> Packages, TArray<FName> InDependentModules, FOutputDevice &HotReloadAr);
+	void DoHotReloadCallback(bool bRecompileFinished, ECompilationResult::Type CompilationResult, TArray<UPackage*> Packages, TArray<FName> InDependentModules, FOutputDevice &HotReloadAr);
 
 	/**
 	 * Gets all currently loaded game module names.
@@ -171,7 +171,7 @@ private:
 	 * The first argument signals whether compilation has finished.
 	 * The second argument shows whether compilation was successful or not.
 	 */
-	DECLARE_DELEGATE_TwoParams( FRecompileModulesCallback, bool, bool );
+	DECLARE_DELEGATE_TwoParams( FRecompileModulesCallback, bool, ECompilationResult::Type );
 
 	/**
 	 * Tries to recompile the specified modules in the background.  When recompiling finishes, the specified callback
@@ -422,7 +422,6 @@ bool FHotReloadModule::RecompileModule( const FName InModuleName, const bool bRe
 	const bool bUseRollingModuleNames = bWasModuleLoaded;
 
 	bool bWasSuccessful = true;
-	FString NewModuleFileNameOnSuccess = FModuleManager::Get().GetModuleFilename(InModuleName);
 	if( bUseRollingModuleNames )
 	{
 		// First, try to compile the module.  If the module is already loaded, we won't unload it quite yet.  Instead
@@ -433,10 +432,6 @@ bool FHotReloadModule::RecompileModule( const FName InModuleName, const bool bRe
 		FString UniqueModuleFileName;
 		FModuleManager::Get().MakeUniqueModuleFilename( InModuleName, UniqueSuffix, UniqueModuleFileName );
 
-		// If the recompile succeeds, we'll update our cached file name to use the new unique file name
-		// that we setup for the module
-		NewModuleFileNameOnSuccess = UniqueModuleFileName;
-
 		TArray< FModuleToRecompile > ModulesToRecompile;
 		FModuleToRecompile ModuleToRecompile;
 		ModuleToRecompile.ModuleName = InModuleName.ToString();
@@ -523,16 +518,16 @@ void FHotReloadModule::DoHotReloadFromEditor()
 	RecordAnalyticsEvent(TEXT("Editor"), Result, Duration, PackagesToRebind.Num(), DependentModules.Num());
 }
 
-void FHotReloadModule::DoHotReloadCallback(bool bRecompileFinished, bool bRecompileSucceeded, TArray<UPackage*> Packages, TArray< FName > InDependentModules, FOutputDevice &HotReloadAr)
+void FHotReloadModule::DoHotReloadCallback(bool bRecompileFinished, ECompilationResult::Type CompilationResult, TArray<UPackage*> Packages, TArray< FName > InDependentModules, FOutputDevice &HotReloadAr)
 {
-	DoHotReloadInternal(bRecompileFinished, bRecompileSucceeded, Packages, InDependentModules, HotReloadAr);
+	DoHotReloadInternal(bRecompileFinished, CompilationResult, Packages, InDependentModules, HotReloadAr);
 }
 
-ECompilationResult::Type FHotReloadModule::DoHotReloadInternal(bool bRecompileFinished, bool bRecompileSucceeded, TArray<UPackage*> Packages, TArray< FName > InDependentModules, FOutputDevice &HotReloadAr)
+ECompilationResult::Type FHotReloadModule::DoHotReloadInternal(bool bRecompileFinished, ECompilationResult::Type CompilationResult, TArray<UPackage*> Packages, TArray< FName > InDependentModules, FOutputDevice &HotReloadAr)
 {
 	ECompilationResult::Type Result = ECompilationResult::Unsupported;
 #if WITH_HOT_RELOAD
-	if (bRecompileSucceeded)
+	if (CompilationResult == ECompilationResult::Succeeded)
 	{
 		FFeedbackContext& ErrorsFC = UClass::GetDefaultPropertiesFeedbackContext();
 		ErrorsFC.Errors.Empty();
@@ -643,7 +638,7 @@ ECompilationResult::Type FHotReloadModule::DoHotReloadInternal(bool bRecompileFi
 		const bool bWasTriggeredAutomatically = !bIsHotReloadingFromEditor;
 		BroadcastHotReload( bWasTriggeredAutomatically );
 	}
-	else if (bRecompileFinished)
+	else if (ECompilationResult::Failed(CompilationResult) && bRecompileFinished)
 	{
 		HotReloadAr.Logf(ELogVerbosity::Warning, TEXT("HotReload failed, recompile failed"));
 		Result = ECompilationResult::OtherCompilationError;
@@ -904,8 +899,8 @@ void FHotReloadModule::DoHotReloadFromIDE()
 		check(PackagesToRebind.Num() || DependentModules.Num())
 		{
 			const bool bRecompileFinished = true;
-			const bool bRecompileSucceeded = true;
-			Result = DoHotReloadInternal(bRecompileFinished, bRecompileSucceeded, PackagesToRebind, DependentModules, *GLog);
+			const ECompilationResult::Type RecompileResult = ECompilationResult::Succeeded;
+			Result = DoHotReloadInternal(bRecompileFinished, RecompileResult, PackagesToRebind, DependentModules, *GLog);
 		}
 
 		GWarn->EndSlowTask();
@@ -1107,7 +1102,8 @@ bool FHotReloadModule::StartCompilingModuleDLLs(const FString& GameName, const T
 	}
 
 	// Shared PCH does no work with hot-reloading modules as we don't scan all modules for them.
-	ExtraArg += TEXT("-nosharedpch ");
+	// Also, if there's nothing to compile, don't bother linking the DLLs as the old ones are up-to-date
+	ExtraArg += TEXT("-nosharedpch -canskiplink ");
 
 	FString TargetName = GameName;
 
@@ -1135,7 +1131,7 @@ bool FHotReloadModule::StartCompilingModuleDLLs(const FString& GameName, const T
 
 		// Fire task completion delegate 
 		
-		RecompileModulesCallback.ExecuteIfBound( false, false );
+		RecompileModulesCallback.ExecuteIfBound( false, ECompilationResult::OtherCompilationError );
 		RecompileModulesCallback.Unbind();
 	}
 
@@ -1237,8 +1233,12 @@ void FHotReloadModule::CheckForFinishedModuleDLLCompile(const bool bWaitForCompl
 		int32 ReturnCode = -1;
 		while (bCompileStillInProgress)
 		{
-			if( FPlatformProcess::GetProcReturnCode( ModuleCompileProcessHandle, &ReturnCode ) )
+			// Store the return code in a temp variable for now because it still gets overwritten
+			// when the process is running.
+			int32 ProcReturnCode = -1;
+			if( FPlatformProcess::GetProcReturnCode( ModuleCompileProcessHandle, &ProcReturnCode ) )
 			{
+				ReturnCode = ProcReturnCode;
 				bCompileStillInProgress = false;
 			}
 			
@@ -1274,12 +1274,13 @@ void FHotReloadModule::CheckForFinishedModuleDLLCompile(const bool bWaitForCompl
 			// Compilation finished, now we need to grab all of the text from the output pipe
 			ModuleCompileReadPipeText += FPlatformProcess::ReadPipe(ModuleCompileReadPipe);
 
-			// The ReturnCode is -1 only if compilation was cancelled.
-			CompilationResult = ReturnCode != -1 ? (ECompilationResult::Type)ReturnCode : ECompilationResult::OtherCompilationError;
+			// This includes 'canceled' (-1) and 'up-to-date' (-2)
+			CompilationResult = (ECompilationResult::Type)ReturnCode;
 
 			// If compilation succeeded for all modules, go back to the modules and update their module file names
 			// in case we recompiled the modules to a new unique file name.  This is needed so that when the module
-			// is reloaded after the recompile, we load the new DLL file name, not the old one
+			// is reloaded after the recompile, we load the new DLL file name, not the old one.
+			// Note that we don't want to do anything in case the build was canceled or source code has not changed.
 			if(CompilationResult == ECompilationResult::Succeeded)
 			{
 				for( int32 CurModuleIndex = 0; CurModuleIndex < ModulesThatWereBeingRecompiled.Num(); ++CurModuleIndex )
@@ -1293,9 +1294,8 @@ void FHotReloadModule::CheckForFinishedModuleDLLCompile(const bool bWaitForCompl
 						OnModuleCompileSucceeded(FName(*CurModule.ModuleName), CurModule.NewModuleFilename);
 					}
 				}
-
-				ModulesThatWereBeingRecompiled.Empty();
 			}
+			ModulesThatWereBeingRecompiled.Empty();
 
 			// We're done with the process handle now
 			ModuleCompileProcessHandle.Close();
@@ -1315,7 +1315,7 @@ void FHotReloadModule::CheckForFinishedModuleDLLCompile(const bool bWaitForCompl
 			// No longer compiling modules
 			ModulesBeingCompiled.Empty();
 
-			bCompileSucceeded = CompilationResult == ECompilationResult::Succeeded;
+			bCompileSucceeded = !ECompilationResult::Failed(CompilationResult);
 
 			if ( bFireEvents )
 			{
@@ -1323,7 +1323,7 @@ void FHotReloadModule::CheckForFinishedModuleDLLCompile(const bool bWaitForCompl
 				ModuleCompilerFinishedEvent.Broadcast(FinalOutput, CompilationResult, !bCompileSucceeded || bShowLogOnSuccess);
 
 				// Fire task completion delegate 
-				RecompileModulesCallback.ExecuteIfBound( true, bCompileSucceeded );
+				RecompileModulesCallback.ExecuteIfBound( true, CompilationResult );
 				RecompileModulesCallback.Unbind();
 			}
 		}
diff --git a/Engine/Source/Editor/MainFrame/Private/MainFrameModule.cpp b/Engine/Source/Editor/MainFrame/Private/MainFrameModule.cpp
--- a/Engine/Source/Editor/MainFrame/Private/MainFrameModule.cpp
+++ b/Engine/Source/Editor/MainFrame/Private/MainFrameModule.cpp
@@ -760,7 +760,7 @@ void FMainFrameModule::HandleLevelEditorModuleCompileFinished(const FString& Log
 		{
 			TArray< FAnalyticsEventAttribute > CompileAttribs;
 			CompileAttribs.Add(FAnalyticsEventAttribute(TEXT("Duration"), FString::Printf(TEXT("%.3f"), ModuleCompileDuration)));
-			CompileAttribs.Add(FAnalyticsEventAttribute(TEXT("Result"), CompilationResult == ECompilationResult::Succeeded ? TEXT("Succeeded") : TEXT("Failed")));
+			CompileAttribs.Add(FAnalyticsEventAttribute(TEXT("Result"), ECompilationResult::ToString(CompilationResult)));
 			FEngineAnalytics::GetProvider().RecordEvent(TEXT("Editor.Modules.Recompile"), CompileAttribs);
 		}
 	}
@@ -769,11 +769,11 @@ void FMainFrameModule::HandleLevelEditorModuleCompileFinished(const FString& Log
 
 	if (NotificationItem.IsValid())
 	{
-		if (CompilationResult == ECompilationResult::Succeeded)
+		if (!ECompilationResult::Failed(CompilationResult))
 		{
 			GEditor->PlayPreviewSound(CompileSuccessSound);
 			NotificationItem->SetText(NSLOCTEXT("MainFrame", "RecompileComplete", "Compile Complete!"));
-			NotificationItem->SetExpireDuration( 1.5f );
+			NotificationItem->SetExpireDuration( 5.0f );
 			NotificationItem->SetCompletionState(SNotificationItem::CS_Success);
 		}
 		else
@@ -792,6 +792,10 @@ void FMainFrameModule::HandleLevelEditorModuleCompileFinished(const FString& Log
 			{
 				NotificationItem->SetText(NSLOCTEXT("MainFrame", "RecompileFailedDueToHeaderChange", "Compile failed due to the header changes. Close the editor and recompile project in IDE to apply changes."));
 			}
+			else if (CompilationResult == ECompilationResult::Canceled)
+			{
+				NotificationItem->SetText(NSLOCTEXT("MainFrame", "RecompileFailed", "Compile Canceled!"));
+			}
 			else
 			{
 				NotificationItem->SetText(NSLOCTEXT("MainFrame", "RecompileFailed", "Compile Failed!"));
@@ -799,6 +803,7 @@ void FMainFrameModule::HandleLevelEditorModuleCompileFinished(const FString& Log
 			
 			NotificationItem->SetCompletionState(SNotificationItem::CS_Fail);
 			NotificationItem->SetHyperlink(FSimpleDelegate::CreateStatic(&Local::ShowCompileLog));
+			NotificationItem->SetExpireDuration(30.0f);
 		}
 
 		NotificationItem->ExpireAndFadeout();
diff --git a/Engine/Source/Programs/UnrealBuildTool/Configuration/UEBuildConfiguration.cs b/Engine/Source/Programs/UnrealBuildTool/Configuration/UEBuildConfiguration.cs
--- a/Engine/Source/Programs/UnrealBuildTool/Configuration/UEBuildConfiguration.cs
+++ b/Engine/Source/Programs/UnrealBuildTool/Configuration/UEBuildConfiguration.cs
@@ -183,6 +183,9 @@ namespace UnrealBuildTool
 		/** True if performing hot-reload from IDE */
 		public static bool bHotReloadFromIDE;
 
+		/** When true, the targets won't execute their link actions if there was nothing to compile */
+		public static bool bSkipLinkingWhenNothingToCompile;
+
 		/// <summary>
 		/// Sets the configuration back to defaults.
 		/// </summary>
@@ -222,6 +225,7 @@ namespace UnrealBuildTool
 			bCompileForSize = false;
 			bHotReloadFromIDE = false;
 			bAllowHotReloadFromIDE = true;
+			bSkipLinkingWhenNothingToCompile = false;
 		}
 
 		/// <summary>
diff --git a/Engine/Source/Programs/UnrealBuildTool/System/ActionGraph.cs b/Engine/Source/Programs/UnrealBuildTool/System/ActionGraph.cs
--- a/Engine/Source/Programs/UnrealBuildTool/System/ActionGraph.cs
+++ b/Engine/Source/Programs/UnrealBuildTool/System/ActionGraph.cs
@@ -259,14 +259,25 @@ namespace UnrealBuildTool
 
 			// Build a list of actions that are both needed for this target and outdated.
 			List<Action> ActionsToExecute = new List<Action>();
+			bool bHasOutdatedNonLinkActions = false;
 			foreach (Action Action in AllActions)
 			{
 				if (Action.CommandPath != null && IsActionOutdatedMap.ContainsKey(Action) && OutdatedActionDictionary[Action])
 				{
 					ActionsToExecute.Add(Action);
+					if (Action.ActionType != ActionType.Link)
+					{
+						bHasOutdatedNonLinkActions = true;
+					}
 				}
 			}
 
+			// Remove link actions if asked to
+			if (UEBuildConfiguration.bSkipLinkingWhenNothingToCompile && !bHasOutdatedNonLinkActions)
+			{
+				ActionsToExecute.Clear();
+			}
+
 			if( BuildConfiguration.bPrintPerformanceInfo )
 			{ 
 				var CheckOutdatednessTime = (DateTime.UtcNow - CheckOutdatednessStartTime).TotalSeconds;
diff --git a/Engine/Source/Programs/UnrealBuildTool/System/ExternalExecution.cs b/Engine/Source/Programs/UnrealBuildTool/System/ExternalExecution.cs
--- a/Engine/Source/Programs/UnrealBuildTool/System/ExternalExecution.cs
+++ b/Engine/Source/Programs/UnrealBuildTool/System/ExternalExecution.cs
@@ -18,6 +18,10 @@ namespace UnrealBuildTool
 	// processes valid.
 	public enum ECompilationResult
 	{
+		/** All targets were up to date, used only with -canskiplink */
+		UpToDate = -2,
+		/** Build was canceled, this is used on the engine side only */
+		Canceled = -1,
 		/** Compilation succeeded */
 		Succeeded = 0,
 		/** Compilation failed because generated code changed which was not supported */
@@ -29,6 +33,13 @@ namespace UnrealBuildTool
 		/** Unknown error */
 		Unknown
 	}
+	public static class CompilationResultExtensions
+	{
+		public static bool Succeeded(this ECompilationResult Result)
+		{
+			return Result == ECompilationResult.Succeeded || Result == ECompilationResult.UpToDate;
+		}
+	}
 
 	/** Information about a module that needs to be passed to UnrealHeaderTool for code generation */
 	[Serializable]
diff --git a/Engine/Source/Programs/UnrealBuildTool/System/UnrealBuildTool.cs b/Engine/Source/Programs/UnrealBuildTool/System/UnrealBuildTool.cs
--- a/Engine/Source/Programs/UnrealBuildTool/System/UnrealBuildTool.cs
+++ b/Engine/Source/Programs/UnrealBuildTool/System/UnrealBuildTool.cs
@@ -932,7 +932,11 @@ namespace UnrealBuildTool
                         else if (LowercaseArg == "-progress")
                         {
                             ProgressWriter.bWriteMarkup = true;
-                        }                        
+                        }
+						else if (LowercaseArg == "-canskiplink")
+						{
+							UEBuildConfiguration.bSkipLinkingWhenNothingToCompile = true;
+						}
                         else if (CheckPlatform.ToString().ToLowerInvariant() == LowercaseArg)
                         {
                             // It's the platform set...
@@ -1487,6 +1491,8 @@ namespace UnrealBuildTool
                     UEBuildConfiguration.bHotReloadFromIDE = UEBuildConfiguration.bAllowHotReloadFromIDE && !Target.bEditorRecompile && ShouldDoHotReload(Target);
                     if (UEBuildConfiguration.bHotReloadFromIDE)
                     {
+						// Don't produce new DLLs if there's been no code changes
+						UEBuildConfiguration.bSkipLinkingWhenNothingToCompile = true;
                         Log.TraceInformation("Compiling game modules for hot reload");
                     }
 
@@ -1642,7 +1648,7 @@ namespace UnrealBuildTool
                             }
                         }
 
-                        if( BuildResult == ECompilationResult.Succeeded )
+                        if( BuildResult.Succeeded() )
                         { 
                             // Make sure any old DLL files from in-engine recompiles aren't lying around.  Must be called after the action graph is finalized.
                             ActionGraph.DeleteStaleHotReloadDLLs();
@@ -1680,6 +1686,10 @@ namespace UnrealBuildTool
                                 {
                                     ToolChain.PostBuildSync(Target);
                                 }
+								if (ActionsToExecute.Count == 0 && UEBuildConfiguration.bSkipLinkingWhenNothingToCompile)
+								{
+									BuildResult = ECompilationResult.UpToDate;
+								}
                             }
                             else
                             {
diff --git a/Engine/Source/Runtime/Core/Public/Misc/CompilationResult.h b/Engine/Source/Runtime/Core/Public/Misc/CompilationResult.h
--- a/Engine/Source/Runtime/Core/Public/Misc/CompilationResult.h
+++ b/Engine/Source/Runtime/Core/Public/Misc/CompilationResult.h
@@ -17,6 +17,10 @@ namespace ECompilationResult
 {
 	enum Type
 	{
+		/** All targets were up to date, used only with -canskiplink */
+		UpToDate = -2,
+		/** Build was canceled, this is used on the engine side only */
+		Canceled = -1,
 		/** Compilation succeeded */
 		Succeeded = 0,
 		/** Compilation failed because generated code changed which was not supported */
@@ -36,6 +40,10 @@ namespace ECompilationResult
 	{
 		switch (Result)
 		{
+		case ECompilationResult::UpToDate:
+			return TEXT("UpToDate");
+		case ECompilationResult::Canceled:
+			return TEXT("Canceled");
 		case ECompilationResult::Succeeded:
 			return TEXT("Succeeded");
 		case ECompilationResult::FailedDueToHeaderChange:
@@ -47,4 +55,12 @@ namespace ECompilationResult
 		};
 		return TEXT("Unknown");
 	}
+
+	/**
+	* Returns false if the provided Result is either UpToDate or Succeeded.
+	*/
+	static FORCEINLINE bool Failed(ECompilationResult::Type Result)
+	{
+		return !(Result == ECompilationResult::Succeeded || Result == ECompilationResult::UpToDate);
+	}
 }
