Hot-Reload: Fixing more build errors when performing Hot-Reload from IDE. Fix for hot-reloading misc game modules from the editor modules list window or through command line.

diff --git a/Engine/Source/Developer/HotReload/Private/HotReload.cpp b/Engine/Source/Developer/HotReload/Private/HotReload.cpp
--- a/Engine/Source/Developer/HotReload/Private/HotReload.cpp
+++ b/Engine/Source/Developer/HotReload/Private/HotReload.cpp
@@ -257,6 +257,9 @@ private:
 	/** New module DLLs */
 	TArray<FRecompiledModule> NewModules;
 
+	/** Moduels that have been recently recompiled from the editor **/
+	TSet<FString> ModulesRecentlyCompiledInTheEditor;
+
 	/** Delegate broadcast when a module has been hot-reloaded */
 	FHotReloadEvent HotReloadEvent;
 
@@ -402,6 +405,14 @@ FString FHotReloadModule::GetModuleCompileMethod(FName InModuleName)
 bool FHotReloadModule::RecompileModule( const FName InModuleName, const bool bReloadAfterRecompile, FOutputDevice &Ar )
 {
 #if WITH_HOT_RELOAD
+	UE_LOG(LogHotReload, Log, TEXT("Recompiling module %s..."), *InModuleName.ToString());
+
+	// This is an internal request for hot-reload (not from IDE)
+	bIsHotReloadingFromEditor = true;
+	// A list of modules that have been recompiled in the editor is going to prevent false
+	// hot-reload from IDE events as this call is blocking any potential callbacks coming from the filesystem
+	// and bIsHotReloadingFromEditor may not be enough to prevent those from being treated as actual hot-reload from IDE modules
+	ModulesRecentlyCompiledInTheEditor.Empty();
 
 	FFormatNamedArguments Args;
 	Args.Add( TEXT("CodeModuleName"), FText::FromName( InModuleName ) );
@@ -437,7 +448,8 @@ bool FHotReloadModule::RecompileModule( const FName InModuleName, const bool bRe
 		ModuleToRecompile.ModuleFileSuffix = UniqueSuffix;
 		ModuleToRecompile.NewModuleFilename = UniqueModuleFileName;
 		ModulesToRecompile.Add( ModuleToRecompile );
-		bWasSuccessful = RecompileModuleDLLs( ModulesToRecompile, Ar );
+		ModulesRecentlyCompiledInTheEditor.Add(FPaths::ConvertRelativePathToFull(UniqueModuleFileName));
+		bWasSuccessful = RecompileModuleDLLs( ModulesToRecompile, Ar );		
 	}
 
 	if( bWasSuccessful )
@@ -456,6 +468,10 @@ bool FHotReloadModule::RecompileModule( const FName InModuleName, const bool bRe
 			FModuleToRecompile ModuleToRecompile;
 			ModuleToRecompile.ModuleName = InModuleName.ToString();
 			ModulesToRecompile.Add( ModuleToRecompile );
+			if (FModuleManager::Get().IsModuleLoaded(InModuleName))
+			{
+				ModulesRecentlyCompiledInTheEditor.Add(FPaths::ConvertRelativePathToFull(FModuleManager::Get().GetModuleFilename(InModuleName)));
+			}
 			bWasSuccessful = RecompileModuleDLLs( ModulesToRecompile, Ar );
 		}
 
@@ -467,6 +483,8 @@ bool FHotReloadModule::RecompileModule( const FName InModuleName, const bool bRe
 		}
 	}
 
+	bIsHotReloadingFromEditor = false;
+
 	return bWasSuccessful;
 #else
 	return false;
@@ -798,7 +816,9 @@ void FHotReloadModule::OnHotReloadBinariesChanged(const TArray<struct FFileChang
 				{
 					for (auto& GameModule : GameModuleNames)
 					{
-						if (Filename.Contains(GameModule) && !NewModules.ContainsByPredicate([&](const FRecompiledModule& Module){ return Module.Name == GameModule; }))
+						if (Filename.Contains(GameModule) && 
+							  !NewModules.ContainsByPredicate([&](const FRecompiledModule& Module){ return Module.Name == GameModule; }) &&
+								!ModulesRecentlyCompiledInTheEditor.Contains(FPaths::ConvertRelativePathToFull(Change.Filename)))
 						{
 							// Add to queue. We do not hot-reload here as there may potentially be other modules being compiled.
 							NewModules.Add(FRecompiledModule(GameModule, Change.Filename));
diff --git a/Engine/Source/Programs/UnrealBuildTool/Configuration/UEBuildTarget.cs b/Engine/Source/Programs/UnrealBuildTool/Configuration/UEBuildTarget.cs
--- a/Engine/Source/Programs/UnrealBuildTool/Configuration/UEBuildTarget.cs
+++ b/Engine/Source/Programs/UnrealBuildTool/Configuration/UEBuildTarget.cs
@@ -1508,7 +1508,7 @@ namespace UnrealBuildTool
 				CreateLinkerFixupsCPPFile();
 			}
 
-			if (!bEditorRecompile)
+			if (!bEditorRecompile && !UEBuildConfiguration.bHotReloadFromIDE)
 			{
 				CreateAutoStartupModuleListGetter();
 			}
