Hot-Reload: Fix for full rebuild when doing hot-reload in the editor for the first time

diff --git a/Engine/Source/Developer/HotReload/Private/HotReload.cpp b/Engine/Source/Developer/HotReload/Private/HotReload.cpp
--- a/Engine/Source/Developer/HotReload/Private/HotReload.cpp
+++ b/Engine/Source/Developer/HotReload/Private/HotReload.cpp
@@ -234,6 +234,9 @@ private:
 	/** Access the module's file and read the timestamp from the file system. Returns true if the timestamp was read successfully. */
 	bool GetModuleFileTimeStamp(FName ModuleName, FDateTime& OutFileTimeStamp) const;
 
+	/** Checks if the specified array of modules to recompile contains only game modules */
+	bool ContainsOnlyGameModules(const TArray< FModuleToRecompile >& ModuleNames) const;
+
 	/** FTicker delegate (hot-reload from IDE) */
 	FTickerDelegate TickerDelegate;
 
@@ -1129,10 +1132,15 @@ bool FHotReloadModule::StartCompilingModuleDLLs(const FString& GameName, const T
 		ExtraArg += TEXT( "-FailIfGeneratedCodeChanges " );
 	}
 
-	// Shared PCH does no work with hot-reloading modules as we don't scan all modules for them.
-	// Also, if there's nothing to compile, don't bother linking the DLLs as the old ones are up-to-date
-	ExtraArg += TEXT("-nosharedpch -canskiplink ");
+	// If there's nothing to compile, don't bother linking the DLLs as the old ones are up-to-date
+	ExtraArg += TEXT("-canskiplink ");
 
+	// Shared PCH does no work with hot-reloading engine/editor modules as we don't scan all modules for them.
+	if (!ContainsOnlyGameModules(ModuleNames))
+	{
+		ExtraArg += TEXT("-nosharedpch ");
+	}
+	
 	FString TargetName = GameName;
 
 #if WITH_EDITOR
@@ -1487,4 +1495,20 @@ bool FHotReloadModule::IsAnyGameModuleLoaded() const
 	return false;
 }
 
+bool FHotReloadModule::ContainsOnlyGameModules(const TArray<FModuleToRecompile>& ModulesToCompile) const
+{
+	const FString AbsoluteGameDir(FPaths::ConvertRelativePathToFull(FPaths::GameDir()));
+	bool bOnlyGameModules = true;
+	for (auto& ModuleToCompile : ModulesToCompile)
+	{
+		const FString FullModulePath(FPaths::ConvertRelativePathToFull(ModuleToCompile.NewModuleFilename));
+		if (!FullModulePath.StartsWith(AbsoluteGameDir))
+		{
+			bOnlyGameModules = false;
+			break;
+		}
+	}
+	return bOnlyGameModules;
+}
+
 #undef LOCTEXT_NAMESPACE
\ No newline at end of file
