Hot-reload: fix potential unwanted recursions when serializing CDO properties.

diff --git a/Engine/Source/Developer/HotReload/Private/HotReloadClassReinstancer.cpp b/Engine/Source/Developer/HotReload/Private/HotReloadClassReinstancer.cpp
--- a/Engine/Source/Developer/HotReload/Private/HotReloadClassReinstancer.cpp
+++ b/Engine/Source/Developer/HotReload/Private/HotReloadClassReinstancer.cpp
@@ -74,11 +74,14 @@ void FHotReloadClassReinstancer::SerializeCDOProperties(UObject* InObject, TArra
 				FName ObjectName = InObj->GetFName();
 				Ar << ClassName;
 				Ar << ObjectName;
-				if (GSerializedProperty && GSerializedProperty->ContainsInstancedObjectProperty() && !VisitedObjects.Contains(InObj))
+				if (!VisitedObjects.Contains(InObj))
 				{
-					// Serialize all DSO properties too
 					VisitedObjects.Add(InObj);
-					FCDOWriter DefaultSubobjectWriter(Bytes, InObj, VisitedObjects);
+					if (GSerializedProperty && GSerializedProperty->ContainsInstancedObjectProperty())
+					{
+						// Serialize all DSO properties too					
+						FCDOWriter DefaultSubobjectWriter(Bytes, InObj, VisitedObjects);
+					}
 				}
 			}
 			else
@@ -126,6 +129,7 @@ void FHotReloadClassReinstancer::SerializeCDOProperties(UObject* InObject, TArra
 		virtual FString GetArchiveName() const override { return TEXT("FCDOWriter"); }
 	};
 	TSet<UObject*> VisitedObjects;
+	VisitedObjects.Add(InObject);
 	FCDOWriter Ar(OutData, InObject, VisitedObjects);
 }
 
