HotReload: Fixing non-unity build

diff --git a/Engine/Source/Developer/HotReload/HotReload.Build.cs b/Engine/Source/Developer/HotReload/HotReload.Build.cs
--- a/Engine/Source/Developer/HotReload/HotReload.Build.cs
+++ b/Engine/Source/Developer/HotReload/HotReload.Build.cs
@@ -9,21 +9,18 @@ public class HotReload : ModuleRules
 				new string[]
 				{
 					"Core",
-					// ... add other public dependencies that you statically link with here ...
+					"CoreUObject",
 				}
-				);
+			);
 
 			PrivateDependencyModuleNames.AddRange(
 				new string[]
 				{
 					"Analytics",
-					"Core",
-					"CoreUObject",
 					"DirectoryWatcher",
 					"DesktopPlatform"
-					// ... add other public dependencies that you statically link with here ...
 				}
-				);
+			);
 
             if (UEBuildConfiguration.bCompileAgainstEngine)
             {
