Hot-Reload: Fixed infinite recursion when serializing CDOs when performing hot-reload

diff --git a/Engine/Source/Developer/HotReload/Private/HotReloadClassReinstancer.cpp b/Engine/Source/Developer/HotReload/Private/HotReloadClassReinstancer.cpp
--- a/Engine/Source/Developer/HotReload/Private/HotReloadClassReinstancer.cpp
+++ b/Engine/Source/Developer/HotReload/Private/HotReloadClassReinstancer.cpp
@@ -108,21 +108,21 @@ void FHotReloadClassReinstancer::SerializeCDOProperties(UObject* InObject, TArra
 		virtual FArchive& operator<<(FLazyObjectPtr& LazyObjectPtr) override
 		{
 			FArchive& Ar = *this;
-			UObject* Obj = LazyObjectPtr.Get();
-			Ar << Obj;
+			auto UniqueID = LazyObjectPtr.GetUniqueID();
+			Ar << UniqueID;
 			return *this;
 		}
 		virtual FArchive& operator<<(FAssetPtr& AssetPtr) override
 		{
 			FArchive& Ar = *this;
-			UObject* Obj = AssetPtr.Get();
-			Ar << Obj;
+			auto UniqueID = AssetPtr.GetUniqueID();
+			Ar << UniqueID;
 			return Ar;
 		}
 		virtual FArchive& operator<<(FStringAssetReference& Value) override
 		{
 			FArchive& Ar = *this;
-			Value.Serialize(Ar);
+			Ar << Value.AssetLongPathname;
 			return Ar;
 		}
 		/** Archive name, for debugging */
