Hot-Reload improvements:

- CDO property values that have changed after hot-reload will now be propagated to the existing instances
- Adding code to a non-code project will no longer require restarting the editor to be able to work with new code (the new module will automatically be compiled and loaded)

diff --git a/Engine/Source/Developer/HotReload/Private/HotReload.cpp b/Engine/Source/Developer/HotReload/Private/HotReload.cpp
--- a/Engine/Source/Developer/HotReload/Private/HotReload.cpp
+++ b/Engine/Source/Developer/HotReload/Private/HotReload.cpp
@@ -13,16 +13,49 @@ DEFINE_LOG_CATEGORY(LogHotReload);
 
 #define LOCTEXT_NAMESPACE "HotReload"
 
+namespace EThreeStateBool
+{
+	enum Type
+	{
+		False,
+		True,
+		Unknown
+	};
+
+	static bool ToBool(EThreeStateBool::Type Value)
+	{
+		switch (Value)
+		{
+		case EThreeStateBool::False:
+			return false;
+		case EThreeStateBool::True:
+			return true;
+		default:
+			UE_LOG(LogHotReload, Fatal, TEXT("Can't convert EThreeStateBool to bool value because it's Unknown"));			
+			break;
+		}
+		return false;
+	}
+
+	static EThreeStateBool::Type FromBool(bool Value)
+	{
+		return Value ? EThreeStateBool::True : EThreeStateBool::False;
+	}
+};
+
 /**
  * Module for HotReload support
  */
 class FHotReloadModule : public IHotReloadModule, FSelfRegisteringExec
 {
 public:
+
 	FHotReloadModule()
 	{
 		ModuleCompileReadPipe = nullptr;
 		bRequestCancelCompilation = false;
+		bIsAnyGameModuleLoaded = EThreeStateBool::Unknown;
+		bDirectoryWatcherInitialized = false;
 	}
 
 	/** IModuleInterface implementation */
@@ -35,17 +68,17 @@ public:
 	/** IHotReloadInterface implementation */
 	virtual void Tick() override;
 	virtual void SaveConfig() override;
-	virtual bool RecompileModule( const FName InModuleName, const bool bReloadAfterRecompile, FOutputDevice &Ar ) override;
+	virtual bool RecompileModule(const FName InModuleName, const bool bReloadAfterRecompile, FOutputDevice &Ar, bool bFailIfGeneratedCodeChanges = true, bool bForceCodeProject = false) override;
 	virtual bool IsCurrentlyCompiling() const override { return ModuleCompileProcessHandle.IsValid(); }
 	virtual void RequestStopCompilation() override { bRequestCancelCompilation = true; }
 	virtual void AddHotReloadFunctionRemap(Native NewFunctionPointer, Native OldFunctionPointer) override;	
 	virtual void RebindPackages(TArray< UPackage* > Packages, TArray< FName > DependentModules, const bool bWaitForCompletion, FOutputDevice &Ar) override;
-	virtual void DoHotReloadFromEditor() override;
+	virtual ECompilationResult::Type DoHotReloadFromEditor() override;
 	virtual FHotReloadEvent& OnHotReload() override { return HotReloadEvent; }	
 	virtual FModuleCompilerStartedEvent& OnModuleCompilerStarted() override { return ModuleCompilerStartedEvent; }
 	virtual FModuleCompilerFinishedEvent& OnModuleCompilerFinished() override { return ModuleCompilerFinishedEvent; }
 	virtual FString GetModuleCompileMethod(FName InModuleName) override;
-	virtual bool IsAnyGameModuleLoaded() const override;
+	virtual bool IsAnyGameModuleLoaded() override;
 
 private:
 	/**
@@ -195,8 +228,9 @@ private:
 	 *
 	 * @param ModuleNames List of modules to recompile, including the module name and optional file suffix.
 	 * @param Ar Output device for logging compilation status.
+	 * @param bForceCodeProject Even if it's a non-code project, treat it as code-based project
 	 */
-	bool RecompileModuleDLLs( const TArray< FModuleToRecompile >& ModuleNames, FOutputDevice& Ar );
+	bool RecompileModuleDLLs(const TArray< FModuleToRecompile >& ModuleNames, FOutputDevice& Ar, bool bFailIfGeneratedCodeChanges, bool bForceCodeProject);
 
 	/** Returns arguments to pass to UnrealBuildTool when compiling modules */
 	static FString MakeUBTArgumentsForModuleCompiling();
@@ -210,11 +244,12 @@ private:
 	 *	@param Ar
 	 *	@param bInFailIfGeneratedCodeChanges If true, fail the compilation if generated headers change.
 	 *	@param InAdditionalCmdLineArgs Additional arguments to pass to UBT.
+	 *  @param bForceCodeProject Compile as code-based project even if there's no game modules loaded
 	 *	@return true if successful, false otherwise.
 	 */
 	bool StartCompilingModuleDLLs(const FString& GameName, const TArray< FModuleToRecompile >& ModuleNames, 
 		const FRecompileModulesCallback& InRecompileModulesCallback, FOutputDevice& Ar, bool bInFailIfGeneratedCodeChanges, 
-		const FString& InAdditionalCmdLineArgs = FString() );
+		const FString& InAdditionalCmdLineArgs, bool bForceCodeProject);
 
 	/** Launches UnrealBuildTool with the specified command line parameters */
 	bool InvokeUnrealBuildToolForCompile(const FString& InCmdLineParams, FOutputDevice &Ar);
@@ -237,6 +272,9 @@ private:
 	/** Checks if the specified array of modules to recompile contains only game modules */
 	bool ContainsOnlyGameModules(const TArray< FModuleToRecompile >& ModuleNames) const;
 
+	/** Callback registered with ModuleManager to know if any new modules have been loaded */
+	void ModulesChangedCallback(FName ModuleName, EModuleChangeReason ReasonForChange);
+
 	/** FTicker delegate (hot-reload from IDE) */
 	FTickerDelegate TickerDelegate;
 
@@ -295,6 +333,12 @@ private:
 
 	/** true if we should attempt to cancel the current async compilation */
 	bool bRequestCancelCompilation;
+
+	/** Tracks the validity of the game module existance */
+	EThreeStateBool::Type bIsAnyGameModuleLoaded;
+
+	/** True if the directory watcher has been successfully initialized */
+	bool bDirectoryWatcherInitialized;
 };
 
 namespace HotReloadDefs
@@ -328,6 +372,8 @@ void FHotReloadModule::StartupModule()
 	// Register hot-reload from IDE ticker
 	TickerDelegate = FTickerDelegate::CreateRaw(this, &FHotReloadModule::Tick);
 	FTicker::GetCoreTicker().AddTicker(TickerDelegate);
+
+	FModuleManager::Get().OnModulesChanged().AddRaw(this, &FHotReloadModule::ModulesChangedCallback);
 }
 
 void FHotReloadModule::ShutdownModule()
@@ -350,7 +396,9 @@ bool FHotReloadModule::Exec( UWorld* Inworld, const TCHAR* Cmd, FOutputDevice& A
 			{
 				const FName ModuleName( *ModuleNameStr );
 				const bool bReloadAfterRecompile = true;
-				RecompileModule( ModuleName, bReloadAfterRecompile, Ar);
+				const bool bForceCodeProject = false;
+				const bool bFailIfGeneratedCodeChanges = true;
+				RecompileModule( ModuleName, bReloadAfterRecompile, Ar, bFailIfGeneratedCodeChanges, bForceCodeProject);
 			}
 
 			return true;
@@ -405,7 +453,7 @@ FString FHotReloadModule::GetModuleCompileMethod(FName InModuleName)
 	}
 }
 
-bool FHotReloadModule::RecompileModule( const FName InModuleName, const bool bReloadAfterRecompile, FOutputDevice &Ar )
+bool FHotReloadModule::RecompileModule(const FName InModuleName, const bool bReloadAfterRecompile, FOutputDevice &Ar, bool bFailIfGeneratedCodeChanges, bool bForceCodeProject)
 {
 #if WITH_HOT_RELOAD
 	UE_LOG(LogHotReload, Log, TEXT("Recompiling module %s..."), *InModuleName.ToString());
@@ -453,7 +501,7 @@ bool FHotReloadModule::RecompileModule( const FName InModuleName, const bool bRe
 		ModuleToRecompile.NewModuleFilename = UniqueModuleFileName;
 		ModulesToRecompile.Add( ModuleToRecompile );
 		ModulesRecentlyCompiledInTheEditor.Add(FPaths::ConvertRelativePathToFull(UniqueModuleFileName));
-		bWasSuccessful = RecompileModuleDLLs( ModulesToRecompile, Ar );		
+		bWasSuccessful = RecompileModuleDLLs(ModulesToRecompile, Ar, bFailIfGeneratedCodeChanges, bForceCodeProject);		
 	}
 
 	if( bWasSuccessful )
@@ -470,24 +518,34 @@ bool FHotReloadModule::RecompileModule( const FName InModuleName, const bool bRe
 			// Try to recompile the DLL
 			TArray< FModuleToRecompile > ModulesToRecompile;
 			FModuleToRecompile ModuleToRecompile;
-			ModuleToRecompile.ModuleName = InModuleName.ToString();
-			ModulesToRecompile.Add( ModuleToRecompile );
+			ModuleToRecompile.ModuleName = InModuleName.ToString();			
 			if (FModuleManager::Get().IsModuleLoaded(InModuleName))
 			{
 				ModulesRecentlyCompiledInTheEditor.Add(FPaths::ConvertRelativePathToFull(FModuleManager::Get().GetModuleFilename(InModuleName)));
 			}
-			bWasSuccessful = RecompileModuleDLLs( ModulesToRecompile, Ar );
+			else
+			{
+				ModuleToRecompile.NewModuleFilename = FModuleManager::Get().GetGameBinariesDirectory() / FModuleManager::GetCleanModuleFilename(InModuleName, true);
+				ModulesRecentlyCompiledInTheEditor.Add(FPaths::ConvertRelativePathToFull(ModuleToRecompile.NewModuleFilename));
+			}
+			ModulesToRecompile.Add( ModuleToRecompile );
+			bWasSuccessful = RecompileModuleDLLs(ModulesToRecompile, Ar, bFailIfGeneratedCodeChanges, bForceCodeProject);
 		}
 
 		// Reload the module if it was loaded before we recompiled
-		if( bWasSuccessful && bWasModuleLoaded && bReloadAfterRecompile )
+		if( bWasSuccessful && (bWasModuleLoaded || bForceCodeProject) && bReloadAfterRecompile )
 		{
 			Ar.Logf( TEXT( "Reloading module after successful compile." ) );
 			bWasSuccessful = FModuleManager::Get().LoadModuleWithCallback( InModuleName, Ar );
 		}
 	}
 
-	bIsHotReloadingFromEditor = false;
+	if (bForceCodeProject && bWasSuccessful)
+	{
+		BroadcastHotReload(false);
+	}
+
+	bIsHotReloadingFromEditor = false;	
 
 	return bWasSuccessful;
 #else
@@ -514,13 +572,20 @@ void FHotReloadModule::AddHotReloadFunctionRemap(Native NewFunctionPointer, Nati
 	HotReloadFunctionRemap.Add(OldFunctionPointer, NewFunctionPointer);
 }
 
-void FHotReloadModule::DoHotReloadFromEditor()
+ECompilationResult::Type FHotReloadModule::DoHotReloadFromEditor()
 {
+	// Get all game modules we want to compile
 	TArray<FString> GameModuleNames;
+	GetGameModules(GameModuleNames);
+	
+	// Don't wait -- we want compiling to happen asynchronously
+	const bool bWaitForCompletion = false;
+
 	TArray<UPackage*> PackagesToRebind;
 	TArray<FName> DependentModules;
-	GetGameModules(GameModuleNames);
+	
 	ECompilationResult::Type Result = ECompilationResult::Unsupported;
+	
 	// Analytics
 	double Duration = 0.0;
 
@@ -528,14 +593,14 @@ void FHotReloadModule::DoHotReloadFromEditor()
 	{
 		FScopedDurationTimer Timer(Duration);
 
-
 		GetPackagesToRebindAndDependentModules(GameModuleNames, PackagesToRebind, DependentModules);
 
-		const bool bWaitForCompletion = false;	// Don't wait -- we want compiling to happen asynchronously
 		Result = RebindPackagesInternal(PackagesToRebind, DependentModules, bWaitForCompletion, *GLog);
 	}
 
 	RecordAnalyticsEvent(TEXT("Editor"), Result, Duration, PackagesToRebind.Num(), DependentModules.Num());
+
+	return Result;
 }
 
 void FHotReloadModule::DoHotReloadCallback(bool bRecompileFinished, ECompilationResult::Type CompilationResult, TArray<UPackage*> Packages, TArray< FName > InDependentModules, FOutputDevice &HotReloadAr)
@@ -767,7 +832,7 @@ void FHotReloadModule::ReinstanceClass(UClass* OldClass, UClass* NewClass)
 	if (ReinstanceHelper.ClassNeedsReinstancing())
 	{
 		UE_LOG(LogHotReload, Log, TEXT("Re-instancing %s after hot-reload."), NewClass ? *NewClass->GetName() : *OldClass->GetName());
-		ReinstanceHelper.ReinstanceObjects();
+		ReinstanceHelper.ReinstanceObjectsAndUpdateDefaults();
 	}
 }
 #endif
@@ -843,8 +908,11 @@ void FHotReloadModule::InitHotReloadWatcher()
 	{
 		// Watch the game binaries folder for new files
 		FString BinariesPath = FPaths::ConvertRelativePathToFull(FPaths::GameDir() / TEXT("Binaries") / FPlatformProcess::GetBinariesSubdirectory());
-		BinariesFolderChangedDelegate = IDirectoryWatcher::FDirectoryChanged::CreateRaw(this, &FHotReloadModule::OnHotReloadBinariesChanged);
-		DirectoryWatcher->RegisterDirectoryChangedCallback(BinariesPath, BinariesFolderChangedDelegate);
+		if (FPaths::DirectoryExists(BinariesPath))
+		{
+			BinariesFolderChangedDelegate = IDirectoryWatcher::FDirectoryChanged::CreateRaw(this, &FHotReloadModule::OnHotReloadBinariesChanged);
+			bDirectoryWatcherInitialized = DirectoryWatcher->RegisterDirectoryChangedCallback(BinariesPath, BinariesFolderChangedDelegate);
+		}
 	}
 }
 
@@ -995,7 +1063,8 @@ bool FHotReloadModule::RecompileModulesAsync( const TArray< FName > ModuleNames,
 	// Kick off compilation!
 	const FString AdditionalArguments = MakeUBTArgumentsForModuleCompiling();
 	const bool bFailIfGeneratedCodeChanges = false;
-	bool bWasSuccessful = StartCompilingModuleDLLs( FApp::GetGameName(), ModulesToRecompile, InRecompileModulesCallback, Ar, bFailIfGeneratedCodeChanges, AdditionalArguments );
+	const bool bForceCodeProject = false;
+	bool bWasSuccessful = StartCompilingModuleDLLs(FApp::GetGameName(), ModulesToRecompile, InRecompileModulesCallback, Ar, bFailIfGeneratedCodeChanges, AdditionalArguments, bForceCodeProject);
 	if (bWasSuccessful)
 	{
 		// Go ahead and check for completion right away.  This is really just so that we can handle the case
@@ -1044,12 +1113,12 @@ void FHotReloadModule::OnModuleCompileSucceeded(FName ModuleName, const FString&
 #endif
 }
 
-bool FHotReloadModule::RecompileModuleDLLs( const TArray< FModuleToRecompile >& ModuleNames, FOutputDevice& Ar )
+bool FHotReloadModule::RecompileModuleDLLs(const TArray< FModuleToRecompile >& ModuleNames, FOutputDevice& Ar, bool bFailIfGeneratedCodeChanges, bool bForceCodeProject)
 {
 	bool bCompileSucceeded = false;
 #if WITH_HOT_RELOAD
 	const FString AdditionalArguments = MakeUBTArgumentsForModuleCompiling();
-	if( StartCompilingModuleDLLs( FApp::GetGameName(), ModuleNames, FRecompileModulesCallback(), Ar, true, AdditionalArguments ) )
+	if (StartCompilingModuleDLLs(FApp::GetGameName(), ModuleNames, FRecompileModulesCallback(), Ar, bFailIfGeneratedCodeChanges, AdditionalArguments, bForceCodeProject))
 	{
 		const bool bWaitForCompletion = true;	// Always wait
 		bool bCompileStillInProgress = false;
@@ -1088,7 +1157,7 @@ FString FHotReloadModule::MakeUBTArgumentsForModuleCompiling()
 
 bool FHotReloadModule::StartCompilingModuleDLLs(const FString& GameName, const TArray< FModuleToRecompile >& ModuleNames, 
 	const FRecompileModulesCallback& InRecompileModulesCallback, FOutputDevice& Ar, bool bInFailIfGeneratedCodeChanges, 
-	const FString& InAdditionalCmdLineArgs )
+	const FString& InAdditionalCmdLineArgs, bool bForceCodeProject)
 {
 #if WITH_HOT_RELOAD
 	// Keep track of what we're compiling
@@ -1147,7 +1216,7 @@ bool FHotReloadModule::StartCompilingModuleDLLs(const FString& GameName, const T
 #if WITH_EDITOR
 	// If there are no game modules loaded, then it's not a code-based project and the target
 	// for UBT should be the editor.
-	if (!IsAnyGameModuleLoaded())
+	if (!bForceCodeProject && !IsAnyGameModuleLoaded())
 	{
 		TargetName = TEXT("UE4Editor");
 	}
@@ -1475,25 +1544,30 @@ bool FHotReloadModule::GetModuleFileTimeStamp(FName ModuleName, FDateTime& OutFi
 	return false;
 }
 
-bool FHotReloadModule::IsAnyGameModuleLoaded() const
+bool FHotReloadModule::IsAnyGameModuleLoaded()
 {
-	// Ask the module manager for a list of currently-loaded gameplay modules
-	TArray< FModuleStatus > ModuleStatuses;
-	FModuleManager::Get().QueryModules(ModuleStatuses);
-
-	for (auto ModuleStatusIt = ModuleStatuses.CreateConstIterator(); ModuleStatusIt; ++ModuleStatusIt)
+	if (bIsAnyGameModuleLoaded == EThreeStateBool::Unknown)
 	{
-		const FModuleStatus& ModuleStatus = *ModuleStatusIt;
+		bool bGameModuleFound = false;
+		// Ask the module manager for a list of currently-loaded gameplay modules
+		TArray< FModuleStatus > ModuleStatuses;
+		FModuleManager::Get().QueryModules(ModuleStatuses);
 
-		// We only care about game modules that are currently loaded
-		if (ModuleStatus.bIsLoaded && ModuleStatus.bIsGameModule)
+		for (auto ModuleStatusIt = ModuleStatuses.CreateConstIterator(); ModuleStatusIt; ++ModuleStatusIt)
 		{
-			// There is at least one loaded game module.
-			return true;
+			const FModuleStatus& ModuleStatus = *ModuleStatusIt;
+
+			// We only care about game modules that are currently loaded
+			if (ModuleStatus.bIsLoaded && ModuleStatus.bIsGameModule)
+			{
+				// There is at least one loaded game module.
+				bGameModuleFound = true;
+				break;
+			}
 		}
+		bIsAnyGameModuleLoaded = EThreeStateBool::FromBool(bGameModuleFound);
 	}
-
-	return false;
+	return EThreeStateBool::ToBool(bIsAnyGameModuleLoaded);
 }
 
 bool FHotReloadModule::ContainsOnlyGameModules(const TArray<FModuleToRecompile>& ModulesToCompile) const
@@ -1512,4 +1586,16 @@ bool FHotReloadModule::ContainsOnlyGameModules(const TArray<FModuleToRecompile>&
 	return bOnlyGameModules;
 }
 
+void FHotReloadModule::ModulesChangedCallback(FName ModuleName, EModuleChangeReason ReasonForChange)
+{
+	// Force update game modules state on the next call to IsAnyGameModuleLoaded
+	bIsAnyGameModuleLoaded = EThreeStateBool::Unknown;
+	
+	// If the hot reload directory watcher hasn't been initialized yet (because the binaries directory did not exist) try to initialize it now
+	if (!bDirectoryWatcherInitialized)
+	{
+		InitHotReloadWatcher();
+	}
+}
+
 #undef LOCTEXT_NAMESPACE
\ No newline at end of file
diff --git a/Engine/Source/Developer/HotReload/Private/HotReloadClassReinstancer.cpp b/Engine/Source/Developer/HotReload/Private/HotReloadClassReinstancer.cpp
--- a/Engine/Source/Developer/HotReload/Private/HotReloadClassReinstancer.cpp
+++ b/Engine/Source/Developer/HotReload/Private/HotReloadClassReinstancer.cpp
@@ -14,6 +14,12 @@ void FHotReloadClassReinstancer::SetupNewClassReinstancing(UClass* InNewClass, U
 	bHasReinstanced = false;
 	bSkipGarbageCollection = false;
 	bNeedsReinstancing = true;
+	NewClass = InNewClass;
+
+	// Collect the original CDO property values
+	SerializeCDOProperties(InOldClass->GetDefaultObject(), OriginalCDOProperties);
+	// Collect the property values of the new CDO
+	SerializeCDOProperties(InNewClass->GetDefaultObject(), ReconstructedCDOProperties);
 
 	SaveClassFieldMapping(InOldClass);
 
@@ -48,22 +54,47 @@ void FHotReloadClassReinstancer::SetupNewClassReinstancing(UClass* InNewClass, U
 	InOldClass->ClassFlags |= CLASS_NewerVersionExists;
 }
 
-void FHotReloadClassReinstancer::SerializeCDOProperties(UObject* InObject, TArray<uint8>& OutData)
+void FHotReloadClassReinstancer::SerializeCDOProperties(UObject* InObject, FHotReloadClassReinstancer::FCDOPropertyData& OutData)
 {
 	// Creates a mem-comparable CDO data
 	class FCDOWriter : public FMemoryWriter
 	{
 		/** Objects already visited by this archive */
 		TSet<UObject*>& VisitedObjects;
+		FCDOPropertyData& PropertyData;
 
 	public:
 		/** Serializes all script properties of the provided DefaultObject */
-		FCDOWriter(TArray<uint8>& InBytes, UObject* DefaultObject, TSet<UObject*>& InVisitedObjects)
-			: FMemoryWriter(InBytes, /* bIsPersistent = */ false, /* bSetOffset = */ true)
+		FCDOWriter(FCDOPropertyData& InOutData, UObject* DefaultObject, TSet<UObject*>& InVisitedObjects)
+			: FMemoryWriter(InOutData.Bytes, /* bIsPersistent = */ false, /* bSetOffset = */ true)
 			, VisitedObjects(InVisitedObjects)
+			, PropertyData(InOutData)
 		{
+			// Disable delta serialization, we want to serialize everything
+			ArNoDelta = true;
 			DefaultObject->SerializeScriptProperties(*this);
 		}
+		virtual void Serialize(void* Data, int64 Num) override
+		{
+			// Collect serialized properties so we can later update their values on instances if they change
+			auto SerializedProperty = GetSerializedProperty();
+			if (SerializedProperty != nullptr)
+			{
+				FCDOProperty& PropertyInfo = PropertyData.Properties.FindOrAdd(SerializedProperty->GetFName());
+				if (PropertyInfo.Property == nullptr)
+				{
+					PropertyInfo.Property = SerializedProperty;
+					PropertyInfo.SerializedValueOffset = Tell();
+					PropertyInfo.SerializedValueSize = Num;
+					PropertyData.Properties.Add(SerializedProperty->GetFName(), PropertyInfo);
+				}
+				else
+				{
+					PropertyInfo.SerializedValueSize += Num;
+				}
+			}
+			FMemoryWriter::Serialize(Data, Num);
+		}
 		/** Serializes an object. Only name and class for normal references, deep serialization for DSOs */
 		virtual FArchive& operator<<(class UObject*& InObj) override
 		{
@@ -80,7 +111,7 @@ void FHotReloadClassReinstancer::SerializeCDOProperties(UObject* InObject, TArra
 					if (Ar.GetSerializedProperty() && Ar.GetSerializedProperty()->ContainsInstancedObjectProperty())
 					{
 						// Serialize all DSO properties too					
-						FCDOWriter DefaultSubobjectWriter(Bytes, InObj, VisitedObjects);
+						FCDOWriter DefaultSubobjectWriter(PropertyData, InObj, VisitedObjects);
 					}
 				}
 			}
@@ -181,21 +212,19 @@ void FHotReloadClassReinstancer::RecreateCDOAndSetupOldClassReinstancing(UClass*
 	bHasReinstanced = false;
 	bSkipGarbageCollection = false;
 	bNeedsReinstancing = false;
+	NewClass = InOldClass; // The class doesn't change in this case
 
 	// Collect the original property values
-	TArray<uint8> OriginalCDOProperties;
 	SerializeCDOProperties(InOldClass->GetDefaultObject(), OriginalCDOProperties);
 	
 	// Destroy and re-create the CDO, re-running its constructor
 	ReconstructClassDefaultObject(InOldClass);
 
 	// Collect the property values after re-constructing the CDO
-	TArray<uint8> ReconstructedCDOProperties;
 	SerializeCDOProperties(InOldClass->GetDefaultObject(), ReconstructedCDOProperties);
 
 	// We only want to re-instance the old class if its CDO's values have changed or any of its DSOs' property values have changed
-	if (OriginalCDOProperties.Num() != ReconstructedCDOProperties.Num() ||
-		FMemory::Memcmp(OriginalCDOProperties.GetData(), ReconstructedCDOProperties.GetData(), OriginalCDOProperties.Num()))
+	if (DefaultPropertiesHaveChanged())
 	{
 		bNeedsReinstancing = true;
 		SaveClassFieldMapping(InOldClass);
@@ -218,6 +247,8 @@ void FHotReloadClassReinstancer::RecreateCDOAndSetupOldClassReinstancing(UClass*
 }
 
 FHotReloadClassReinstancer::FHotReloadClassReinstancer(UClass* InNewClass, UClass* InOldClass)
+	: NewClass(nullptr)
+	, bNeedsReinstancing(false)
 {
 	// If InNewClass is NULL, then the old class has not changed after hot-reload.
 	// However, we still need to check for changes to its constructor code (CDO values).
@@ -238,4 +269,74 @@ FHotReloadClassReinstancer::~FHotReloadClassReinstancer()
 	DuplicatedClass = NULL;
 }
 
+void FHotReloadClassReinstancer::UpdateDefaultProperties()
+{
+	struct FPropertyToUpdate
+	{
+		UProperty* Property;
+		uint8* OldSerializedValuePtr;
+		uint8* NewValuePtr;
+		int64 OldSerializedSize;
+	};
+	TArray<FPropertyToUpdate> PropertiesToUpdate;
+	// Collect all properties that have actually changed
+	for (auto& Pair : ReconstructedCDOProperties.Properties)
+	{
+		auto OldPropertyInfo = OriginalCDOProperties.Properties.Find(Pair.Key);
+		if (OldPropertyInfo)
+		{
+			auto& NewPropertyInfo = Pair.Value;
+			if (NewPropertyInfo.Property->GetOuter() == NewClass)
+			{
+				uint8* OldSerializedValuePtr = OriginalCDOProperties.Bytes.GetData() + OldPropertyInfo->SerializedValueOffset;
+				uint8* NewSerializedValuePtr = ReconstructedCDOProperties.Bytes.GetData() + NewPropertyInfo.SerializedValueOffset;
+				if (OldPropertyInfo->SerializedValueSize != NewPropertyInfo.SerializedValueSize ||
+					FMemory::Memcmp(OldSerializedValuePtr, NewSerializedValuePtr, OldPropertyInfo->SerializedValueSize) != 0)
+				{
+					// Property value has changed so add it to the list of properties that need updating on instances
+					FPropertyToUpdate PropertyToUpdate;
+					PropertyToUpdate.Property = NewPropertyInfo.Property;
+					PropertyToUpdate.OldSerializedValuePtr = OldSerializedValuePtr;
+					PropertyToUpdate.NewValuePtr = PropertyToUpdate.Property->ContainerPtrToValuePtr<uint8>(NewClass->GetDefaultObject());
+					PropertyToUpdate.OldSerializedSize = OldPropertyInfo->SerializedValueSize;
+					PropertiesToUpdate.Add(PropertyToUpdate);
+				}
+			}
+		}
+	}
+	if (PropertiesToUpdate.Num())
+	{
+		TArray<uint8> CurrentValueSerializedData;		
+
+		// Update properties on all existing instances of the class
+		for (FObjectIterator It(NewClass); It; ++It)
+		{
+			UObject* ObjectPtr = *It;
+			for (auto& PropertyToUpdate : PropertiesToUpdate)
+			{
+				uint8* InstanceValuePtr = PropertyToUpdate.Property->ContainerPtrToValuePtr<uint8>(ObjectPtr);
+
+				// Serialize current value to a byte array as we don't have the previous CDO to compare against, we only have its serialized property data
+				CurrentValueSerializedData.Empty(CurrentValueSerializedData.Num() + CurrentValueSerializedData.GetSlack());
+				FMemoryWriter CurrentValueWriter(CurrentValueSerializedData);
+				PropertyToUpdate.Property->SerializeItem(CurrentValueWriter, InstanceValuePtr);
+				
+				// Update only when the current value on the instance is identical to the original CDO
+				if (CurrentValueSerializedData.Num() == PropertyToUpdate.OldSerializedSize &&
+					FMemory::Memcmp(CurrentValueSerializedData.GetData(), PropertyToUpdate.OldSerializedValuePtr, CurrentValueSerializedData.Num()) == 0)
+				{
+					// Update with the new value
+					PropertyToUpdate.Property->CopyCompleteValue(InstanceValuePtr, PropertyToUpdate.NewValuePtr);
+				}
+			}
+		}
+	}
+}
+
+void FHotReloadClassReinstancer::ReinstanceObjectsAndUpdateDefaults()
+{
+	ReinstanceObjects();
+	UpdateDefaultProperties();
+}
+
 #endif
\ No newline at end of file
diff --git a/Engine/Source/Developer/HotReload/Private/HotReloadClassReinstancer.h b/Engine/Source/Developer/HotReload/Private/HotReloadClassReinstancer.h
--- a/Engine/Source/Developer/HotReload/Private/HotReloadClassReinstancer.h
+++ b/Engine/Source/Developer/HotReload/Private/HotReloadClassReinstancer.h
@@ -10,8 +10,38 @@
  */
 class FHotReloadClassReinstancer : public FBlueprintCompileReinstancer
 {
+	/** Holds a property and its offset in the serialized properties data array */
+	struct FCDOProperty
+	{
+		FCDOProperty()
+			: Property(nullptr)
+			, SerializedValueOffset(0)
+			, SerializedValueSize(0)
+		{}
+
+		UProperty* Property;
+		int64 SerializedValueOffset;
+		int64 SerializedValueSize;
+	};
+
+	/** Contains all serialized CDO property data and the map of all serialized properties */
+	struct FCDOPropertyData
+	{
+		TArray<uint8> Bytes;
+		TMap<FName, FCDOProperty> Properties;
+	};
+
+	/** Hot-reloaded version of the old class */
+	UClass* NewClass;
+
+	/** Serialized properties of the original CDO (before hot-reload) */
+	FCDOPropertyData OriginalCDOProperties;
+
+	/** Serialized properties of the new CDO (after hot-reload) */
+	FCDOPropertyData ReconstructedCDOProperties;
+
 	/** True if the provided native class needs re-instancing */
-	bool bNeedsReinstancing;
+	bool bNeedsReinstancing;	
 
 	/** 
 	 * Sets the re-instancer up for new class re-instancing 
@@ -34,7 +64,7 @@ class FHotReloadClassReinstancer : public FBlueprintCompileReinstancer
 	* @param InObject CDO
 	* @param OutData Data containing all of the CDO property values
 	*/
-	void SerializeCDOProperties(UObject* InObject, TArray<uint8>& OutData);
+	void SerializeCDOProperties(UObject* InObject, FCDOPropertyData& OutData);
 
 	/**
 	* Re-creates class default object
@@ -43,6 +73,16 @@ class FHotReloadClassReinstancer : public FBlueprintCompileReinstancer
 	*/
 	void ReconstructClassDefaultObject(UClass* InOldClass);
 
+	/** Updates property values on instances of the hot-reloaded class */
+	void UpdateDefaultProperties();
+
+	/** Returns true if the properties of the CDO have changed during hot-reload */
+	FORCEINLINE bool DefaultPropertiesHaveChanged() const
+	{
+		return OriginalCDOProperties.Bytes.Num() != ReconstructedCDOProperties.Bytes.Num() ||
+			FMemory::Memcmp(OriginalCDOProperties.Bytes.GetData(), ReconstructedCDOProperties.Bytes.GetData(), OriginalCDOProperties.Bytes.Num());
+	}
+
 public:
 
 	/** Sets the re-instancer up to re-instance native classes */
@@ -56,6 +96,9 @@ public:
 	{
 		return bNeedsReinstancing;
 	}
+
+	/** Reinstances all objects of the hot-reloaded class and update their properties to match the new CDO */
+	void ReinstanceObjectsAndUpdateDefaults();
 };
 
 #endif // WITH_ENGINE
\ No newline at end of file
diff --git a/Engine/Source/Developer/ModuleUI/Private/SModuleUI.cpp b/Engine/Source/Developer/ModuleUI/Private/SModuleUI.cpp
--- a/Engine/Source/Developer/ModuleUI/Private/SModuleUI.cpp
+++ b/Engine/Source/Developer/ModuleUI/Private/SModuleUI.cpp
@@ -263,7 +263,9 @@ FReply SModuleUI::FModuleListItem::OnRecompileClicked()
 		{
 			// Perform a regular unload, then reload
 			const bool bReloadAfterRecompile = true;
-			HotReloadSupport.RecompileModule( ModuleName, bReloadAfterRecompile, *GLog );
+			const bool bForceCodeProject = false;
+			const bool bFailIfGeneratedCodeChanges = true;
+			HotReloadSupport.RecompileModule(ModuleName, bReloadAfterRecompile, *GLog, bFailIfGeneratedCodeChanges, bForceCodeProject);
 		}
 	}
 
diff --git a/Engine/Source/Editor/GameProjectGeneration/Private/GameProjectUtils.cpp b/Engine/Source/Editor/GameProjectGeneration/Private/GameProjectUtils.cpp
--- a/Engine/Source/Editor/GameProjectGeneration/Private/GameProjectUtils.cpp
+++ b/Engine/Source/Editor/GameProjectGeneration/Private/GameProjectUtils.cpp
@@ -21,6 +21,7 @@
 #include "SNotificationList.h"
 #include "NotificationManager.h"
 #include "GameFramework/GameMode.h"
+#include "HotReloadInterface.h"
 
 #define LOCTEXT_NAMESPACE "GameProjectUtils"
 
@@ -2844,10 +2845,11 @@ bool GameProjectUtils::AddCodeToProject_Internal(const FString& NewClassName, co
 	SlowTask.MakeDialog();
 
 	SlowTask.EnterProgressFrame();
-
+	
 	// If the project does not already contain code, add the primary game module
 	TArray<FString> CreatedFiles;
-	if ( !ProjectHasCodeFiles() )
+	bool bDidNotHaveAnyCodeFiles = !ProjectHasCodeFiles();
+	if (bDidNotHaveAnyCodeFiles)
 	{
 		// We always add the basic source code to the root directory, not the potential sub-directory provided by NewClassPath
 		const FString SourceDir = FPaths::GameSourceDir().LeftChop(1); // Trim the trailing /
@@ -2927,6 +2929,21 @@ bool GameProjectUtils::AddCodeToProject_Internal(const FString& NewClassName, co
 	OutHeaderFilePath = NewHeaderFilename;
 	OutCppFilePath = NewCppFilename;
 
+	if (bDidNotHaveAnyCodeFiles)
+	{
+		// This is the first time we add code to this project so compile its game DLL
+		const FString GameModuleName = FApp::GetGameName();
+		IHotReloadInterface& HotReloadSupport = FModuleManager::LoadModuleChecked<IHotReloadInterface>("HotReload");
+		const bool bReloadAfterCompiling = true;
+		const bool bForceCodeProject = true;
+		const bool bFailIfGeneratedCodeChanges = false;
+		if (!HotReloadSupport.RecompileModule(*GameModuleName, bReloadAfterCompiling, *GWarn, bFailIfGeneratedCodeChanges, bForceCodeProject))
+		{
+			OutFailReason = LOCTEXT("FailedToCompileNewGameModule", "Failed to compile newly created game module.");
+			return false;
+		}
+	}
+
 	return true;
 }
 
diff --git a/Engine/Source/Editor/LevelEditor/Private/LevelEditorToolBar.cpp b/Engine/Source/Editor/LevelEditor/Private/LevelEditorToolBar.cpp
--- a/Engine/Source/Editor/LevelEditor/Private/LevelEditorToolBar.cpp
+++ b/Engine/Source/Editor/LevelEditor/Private/LevelEditorToolBar.cpp
@@ -1154,10 +1154,15 @@ TSharedRef< SWidget > FLevelEditorToolBar::MakeLevelEditorToolBar( const TShared
 			true);
 
 		// Only show the compile options on machines with the solution (assuming they can build it)
-		if ( FSourceCodeNavigation::IsCompilerAvailable() && FLevelEditorActionCallbacks::CanShowSourceCodeActions() )
+		if ( FSourceCodeNavigation::IsCompilerAvailable() )
 		{
+			// Since we can always add new code to the project, only hide these buttons if we haven't done so yet
 			ToolbarBuilder.AddToolBarButton(
-				FLevelEditorCommands::Get().RecompileGameCode,
+				FUIAction(
+					FExecuteAction::CreateStatic(&FLevelEditorActionCallbacks::RecompileGameCode_Clicked),
+					FCanExecuteAction::CreateStatic(&FLevelEditorActionCallbacks::Recompile_CanExecute),
+					FIsActionChecked(),
+					FIsActionButtonVisible::CreateStatic(FLevelEditorActionCallbacks::CanShowSourceCodeActions)),
 				NAME_None,
 				LOCTEXT( "CompileMenuButton", "Compile" ),
 				FLevelEditorCommands::Get().RecompileGameCode->GetDescription(),
@@ -1165,7 +1170,11 @@ TSharedRef< SWidget > FLevelEditorToolBar::MakeLevelEditorToolBar( const TShared
 				);
 
 			ToolbarBuilder.AddComboButton(
-				FUIAction(),
+				FUIAction(
+					FExecuteAction(),
+					FCanExecuteAction(),
+					FIsActionChecked(),
+					FIsActionButtonVisible::CreateStatic(FLevelEditorActionCallbacks::CanShowSourceCodeActions)),
 				FOnGetContent::CreateStatic( &FLevelEditorToolBar::GenerateCompileMenuContent, InCommandList ),
 				LOCTEXT( "CompileCombo_Label", "Compile Options" ),
 				LOCTEXT( "CompileMenuCombo_ToolTip", "Recompile options" ),
diff --git a/Engine/Source/Runtime/Core/Private/Misc/Paths.cpp b/Engine/Source/Runtime/Core/Private/Misc/Paths.cpp
--- a/Engine/Source/Runtime/Core/Private/Misc/Paths.cpp
+++ b/Engine/Source/Runtime/Core/Private/Misc/Paths.cpp
@@ -456,6 +456,11 @@ bool FPaths::FileExists(const FString& InPath)
 	return IFileManager::Get().GetTimeStamp(*InPath) > FDateTime::MinValue();
 }
 
+bool FPaths::DirectoryExists(const FString& InPath)
+{
+	return IFileManager::Get().DirectoryExists(*InPath);
+}
+
 bool FPaths::IsDrive(const FString& InPath)
 {
 	FString ConvertedPathString = InPath;
diff --git a/Engine/Source/Runtime/Core/Private/Modules/ModuleManager.cpp b/Engine/Source/Runtime/Core/Private/Modules/ModuleManager.cpp
--- a/Engine/Source/Runtime/Core/Private/Modules/ModuleManager.cpp
+++ b/Engine/Source/Runtime/Core/Private/Modules/ModuleManager.cpp
@@ -746,7 +746,13 @@ FString FModuleManager::GetModuleFilename(FName ModuleName) const
 
 void FModuleManager::SetModuleFilename(FName ModuleName, const FString& Filename)
 {
-	Modules.FindChecked(ModuleName)->Filename = Filename;
+	auto& Module = Modules.FindChecked(ModuleName);
+	Module->Filename = Filename;
+	// If it's a new module then also update its OriginalFilename
+	if (Module->OriginalFilename.IsEmpty())
+	{
+		Module->OriginalFilename = Filename;
+	}
 }
 
 
@@ -976,6 +982,14 @@ void FModuleManager::SetGameBinariesDirectory(const TCHAR* InDirectory)
 #endif
 }
 
+FString FModuleManager::GetGameBinariesDirectory() const
+{
+	if (GameBinariesDirectories.Num())
+	{
+		return GameBinariesDirectories[0];
+	}
+	return FString();
+}
 
 bool FModuleManager::DoesLoadedModuleHaveUObjects( const FName ModuleName )
 {
diff --git a/Engine/Source/Runtime/Core/Public/Misc/Paths.h b/Engine/Source/Runtime/Core/Public/Misc/Paths.h
--- a/Engine/Source/Runtime/Core/Public/Misc/Paths.h
+++ b/Engine/Source/Runtime/Core/Public/Misc/Paths.h
@@ -312,6 +312,9 @@ public:
 	/** @return true if this file was found, false otherwise */
 	static bool FileExists(const FString& InPath);
 
+	/** @return true if this directory was found, false otherwise */
+	static bool DirectoryExists(const FString& InPath);
+
 	/** @return true if this path represents a drive */
 	static bool IsDrive(const FString& InPath);
 
diff --git a/Engine/Source/Runtime/Core/Public/Modules/ModuleManager.h b/Engine/Source/Runtime/Core/Public/Modules/ModuleManager.h
--- a/Engine/Source/Runtime/Core/Public/Modules/ModuleManager.h
+++ b/Engine/Source/Runtime/Core/Public/Modules/ModuleManager.h
@@ -359,6 +359,11 @@ public:
 	 */
 	void SetGameBinariesDirectory(const TCHAR* InDirectory);
 
+	/**
+	*	Gets the game binaries directory
+	*/
+	FString GetGameBinariesDirectory() const;
+
 	/**
 	 * Checks to see if the specified module exists and is compatible with the current engine version. 
 	 *
diff --git a/Engine/Source/Runtime/Core/Public/Serialization/Archive.h b/Engine/Source/Runtime/Core/Public/Serialization/Archive.h
--- a/Engine/Source/Runtime/Core/Public/Serialization/Archive.h
+++ b/Engine/Source/Runtime/Core/Public/Serialization/Archive.h
@@ -85,7 +85,7 @@ public:
 		}
 	}
 
-	void Serialize( void* Data, int64 Num )
+	virtual void Serialize(void* Data, int64 Num) override
 	{
 		const int64 NumBytesToAdd = Offset + Num - Bytes.Num();
 		if( NumBytesToAdd > 0 )
diff --git a/Engine/Source/Runtime/CoreUObject/Public/Misc/HotReloadInterface.h b/Engine/Source/Runtime/CoreUObject/Public/Misc/HotReloadInterface.h
--- a/Engine/Source/Runtime/CoreUObject/Public/Misc/HotReloadInterface.h
+++ b/Engine/Source/Runtime/CoreUObject/Public/Misc/HotReloadInterface.h
@@ -41,8 +41,13 @@ public:
 
 	/**
 	 * Recompiles a single module
+	 *
+	 * @param InModuleName Name of the module to compile
+	 * @param bReloadAfterRecompile Should the module be reloaded after recompile
+	 * @param Ar Output device (logging)
+	 * @param bForceCodeProject Even if this is not code-based project compile with game project as the target for UBT (do not use UE4Editor target)
 	 */
-	virtual bool RecompileModule( const FName InModuleName, const bool bReloadAfterRecompile, FOutputDevice &Ar ) = 0;
+	virtual bool RecompileModule(const FName InModuleName, const bool bReloadAfterRecompile, FOutputDevice &Ar, bool bFailIfGeneratedCodeChanges = true, bool bForceCodeProject = false) = 0;
 
 	/**
 	 * Returns whether modules are currently being compiled
@@ -60,9 +65,9 @@ public:
 	virtual void AddHotReloadFunctionRemap(Native NewFunctionPointer, Native OldFunctionPointer) = 0;
 
 	/**
-	* Performs hot reload from the editor
+	* Performs hot reload from the editor of all currently loaded game modules.
 	*/
-	virtual void DoHotReloadFromEditor() = 0;
+	virtual ECompilationResult::Type DoHotReloadFromEditor() = 0;
 
 	/**
 	* HotReload: Reloads the DLLs for given packages.
@@ -102,6 +107,6 @@ public:
 	/**
 	 * Checks if there's any game modules currently loaded
 	 */
-	virtual bool IsAnyGameModuleLoaded() const = 0;
+	virtual bool IsAnyGameModuleLoaded() = 0;
 };
 
diff --git a/Engine/Source/Runtime/Engine/Private/RendererSupport.cpp b/Engine/Source/Runtime/Engine/Private/RendererSupport.cpp
--- a/Engine/Source/Runtime/Engine/Private/RendererSupport.cpp
+++ b/Engine/Source/Runtime/Engine/Private/RendererSupport.cpp
@@ -95,7 +95,9 @@ void RecompileRendererModule()
 		bool bCompiledSuccessfully = false;
 		do 
 		{
-			bCompiledSuccessfully = HotReload->RecompileModule(RendererModuleName, false, *GLog);
+			const bool bForceCodeProject = false;
+			const bool bFailIfGeneratedCodeChanges = true;
+			bCompiledSuccessfully = HotReload->RecompileModule(RendererModuleName, false, *GLog, bFailIfGeneratedCodeChanges, bForceCodeProject);
 
 			if (!bCompiledSuccessfully)
 			{
