Fixed hot reload not actually reloading game modules in some cases
- When hot reload was invoked from the IDE, the editor would not reload all of the game modules if there was more than one module had the game name as part of the module name
- For example after compiling "VehicleGame" and "VehicleGameLoadingScreen", only "VehicleGameLoadingScreen" would be reloaded in the editor

diff --git a/Engine/Source/Developer/HotReload/Private/HotReload.cpp b/Engine/Source/Developer/HotReload/Private/HotReload.cpp
--- a/Engine/Source/Developer/HotReload/Private/HotReload.cpp
+++ b/Engine/Source/Developer/HotReload/Private/HotReload.cpp
@@ -157,9 +157,9 @@ private:
 	void DoHotReloadCallback(bool bRecompileFinished, ECompilationResult::Type CompilationResult, TArray<UPackage*> Packages, TArray<FName> InDependentModules, FOutputDevice &HotReloadAr);
 
 	/**
-	 * Gets all currently loaded game module names.
+	 * Gets all currently loaded game module names and optionally, the file names for those modules
 	 */
-	void GetGameModules(TArray<FString>& OutGameModules);
+	void GetGameModules(TArray<FString>& OutGameModules, TArray<FString>* OutGameModuleFilePaths = nullptr);
 
 	/**
 	 * Gets packages to re-bind and dependent modules.
@@ -837,7 +837,7 @@ void FHotReloadModule::ReinstanceClass(UClass* OldClass, UClass* NewClass)
 }
 #endif
 
-void FHotReloadModule::GetGameModules(TArray<FString>& OutGameModules)
+void FHotReloadModule::GetGameModules(TArray<FString>& OutGameModules, TArray<FString>* OutGameModuleFilePaths )
 {
 	// Ask the module manager for a list of currently-loaded gameplay modules
 	TArray< FModuleStatus > ModuleStatuses;
@@ -850,7 +850,11 @@ void FHotReloadModule::GetGameModules(TArray<FString>& OutGameModules)
 		// We only care about game modules that are currently loaded
 		if (ModuleStatus.bIsLoaded && ModuleStatus.bIsGameModule)
 		{
-			OutGameModules.AddUnique(ModuleStatus.Name);
+			OutGameModules.Add(ModuleStatus.Name);
+			if( OutGameModuleFilePaths != nullptr )
+			{
+				OutGameModuleFilePaths->Add(ModuleStatus.FilePath);
+			}
 		}
 	}
 }
@@ -864,7 +868,8 @@ void FHotReloadModule::OnHotReloadBinariesChanged(const TArray<struct FFileChang
 	}
 
 	TArray< FString > GameModuleNames;
-	GetGameModules(GameModuleNames);
+	TArray< FString > GameModuleFilePaths;
+	GetGameModules(GameModuleNames, &GameModuleFilePaths);
 
 	if (GameModuleNames.Num() > 0)
 	{
@@ -883,15 +888,21 @@ void FHotReloadModule::OnHotReloadBinariesChanged(const TArray<struct FFileChang
 				const FString Filename = FPaths::GetCleanFilename(Change.Filename);
 				if (Filename.EndsWith(FPlatformProcess::GetModuleExtension()))
 				{
-					for (auto& GameModule : GameModuleNames)
+					for( int32 GameModuleIndex = 0; GameModuleIndex < GameModuleNames.Num(); ++GameModuleIndex )
 					{
-						if (Filename.Contains(GameModule) && 
-							  !NewModules.ContainsByPredicate([&](const FRecompiledModule& Module){ return Module.Name == GameModule; }) &&
-								!ModulesRecentlyCompiledInTheEditor.Contains(FPaths::ConvertRelativePathToFull(Change.Filename)))
+						const FString& GameModuleName = GameModuleNames[ GameModuleIndex ];
+						const FString& GameModuleFilePath = GameModuleFilePaths[ GameModuleIndex ];
+
+						const FString GameModuleFileNameWithoutExtension = FPaths::GetBaseFilename( GameModuleFilePath );
+						if( Filename.StartsWith( GameModuleFileNameWithoutExtension + TEXT( "-" ) ) )	// Hot reload always adds a numbered suffix preceded by a hyphen, but otherwise the module name must match exactly!
 						{
-							// Add to queue. We do not hot-reload here as there may potentially be other modules being compiled.
-							NewModules.Add(FRecompiledModule(GameModule, Change.Filename));
-							UE_LOG(LogHotReload, Log, TEXT("New module detected: %s"), *Filename);
+							if ( !NewModules.ContainsByPredicate([&](const FRecompiledModule& Module){ return Module.Name == GameModuleName; }) &&
+								 !ModulesRecentlyCompiledInTheEditor.Contains(FPaths::ConvertRelativePathToFull(Change.Filename)))
+							{
+								// Add to queue. We do not hot-reload here as there may potentially be other modules being compiled.
+								NewModules.Add(FRecompiledModule(GameModuleName, Change.Filename));
+								UE_LOG(LogHotReload, Log, TEXT("New module detected: %s"), *Filename);
+							}
 						}
 					}
 				}
