Fixing hot-reload not updating values on DSO's properties.

UE-8067

diff --git a/Engine/Source/Developer/HotReload/Private/HotReloadClassReinstancer.cpp b/Engine/Source/Developer/HotReload/Private/HotReloadClassReinstancer.cpp
--- a/Engine/Source/Developer/HotReload/Private/HotReloadClassReinstancer.cpp
+++ b/Engine/Source/Developer/HotReload/Private/HotReloadClassReinstancer.cpp
@@ -61,14 +61,18 @@ void FHotReloadClassReinstancer::SerializeCDOProperties(UObject* InObject, FHotR
 	{
 		/** Objects already visited by this archive */
 		TSet<UObject*>& VisitedObjects;
+		/** Output property data */
 		FCDOPropertyData& PropertyData;
+		/** Current subobject being serialized */
+		FName SubobjectName;
 
 	public:
 		/** Serializes all script properties of the provided DefaultObject */
-		FCDOWriter(FCDOPropertyData& InOutData, UObject* DefaultObject, TSet<UObject*>& InVisitedObjects)
+		FCDOWriter(FCDOPropertyData& InOutData, UObject* DefaultObject, TSet<UObject*>& InVisitedObjects, FName InSubobjectName = NAME_None)
 			: FMemoryWriter(InOutData.Bytes, /* bIsPersistent = */ false, /* bSetOffset = */ true)
 			, VisitedObjects(InVisitedObjects)
 			, PropertyData(InOutData)
+			, SubobjectName(InSubobjectName)
 		{
 			// Disable delta serialization, we want to serialize everything
 			ArNoDelta = true;
@@ -84,6 +88,7 @@ void FHotReloadClassReinstancer::SerializeCDOProperties(UObject* InObject, FHotR
 				if (PropertyInfo.Property == nullptr)
 				{
 					PropertyInfo.Property = SerializedProperty;
+					PropertyInfo.SubobjectName = SubobjectName;
 					PropertyInfo.SerializedValueOffset = Tell();
 					PropertyInfo.SerializedValueSize = Num;
 					PropertyData.Properties.Add(SerializedProperty->GetFName(), PropertyInfo);
@@ -111,7 +116,7 @@ void FHotReloadClassReinstancer::SerializeCDOProperties(UObject* InObject, FHotR
 					if (Ar.GetSerializedProperty() && Ar.GetSerializedProperty()->ContainsInstancedObjectProperty())
 					{
 						// Serialize all DSO properties too					
-						FCDOWriter DefaultSubobjectWriter(PropertyData, InObj, VisitedObjects);
+						FCDOWriter DefaultSubobjectWriter(PropertyData, InObj, VisitedObjects, InObj->GetFName());
 					}
 				}
 			}
@@ -266,7 +271,20 @@ FHotReloadClassReinstancer::~FHotReloadClassReinstancer()
 {
 	// Make sure the base class does not remove the DuplicatedClass from root, we not always want it.
 	// For example when we're just reconstructing CDOs. Other cases are handled by HotReloadClassReinstancer.
-	DuplicatedClass = NULL;
+	DuplicatedClass = nullptr;
+}
+
+/** Helper for finding subobject in an array. Usually there's not that many subobjects on a class to justify a TMap */
+FORCEINLINE static UObject* FindDefaultSubobject(TArray<UObject*>& InDefaultSubobjects, FName SubobjectName)
+{
+	for (auto Subobject : InDefaultSubobjects)
+	{
+		if (Subobject->GetFName() == SubobjectName)
+		{
+			return Subobject;
+		}
+	}
+	return nullptr;
 }
 
 void FHotReloadClassReinstancer::UpdateDefaultProperties()
@@ -274,10 +292,18 @@ void FHotReloadClassReinstancer::UpdateDefaultProperties()
 	struct FPropertyToUpdate
 	{
 		UProperty* Property;
+		FName SubobjectName;
 		uint8* OldSerializedValuePtr;
 		uint8* NewValuePtr;
 		int64 OldSerializedSize;
 	};
+
+	// Collect default subobjects to update their properties too
+	const int32 DefaultSubobjectArrayCapacity = 16;
+	TArray<UObject*> DefaultSubobjectArray;
+	DefaultSubobjectArray.Empty(DefaultSubobjectArrayCapacity);
+	NewClass->GetDefaultObject()->CollectDefaultSubobjects(DefaultSubobjectArray);
+
 	TArray<FPropertyToUpdate> PropertiesToUpdate;
 	// Collect all properties that have actually changed
 	for (auto& Pair : ReconstructedCDOProperties.Properties)
@@ -286,21 +312,37 @@ void FHotReloadClassReinstancer::UpdateDefaultProperties()
 		if (OldPropertyInfo)
 		{
 			auto& NewPropertyInfo = Pair.Value;
-			if (NewPropertyInfo.Property->GetOuter() == NewClass)
+
+			uint8* OldSerializedValuePtr = OriginalCDOProperties.Bytes.GetData() + OldPropertyInfo->SerializedValueOffset;
+			uint8* NewSerializedValuePtr = ReconstructedCDOProperties.Bytes.GetData() + NewPropertyInfo.SerializedValueOffset;
+			if (OldPropertyInfo->SerializedValueSize != NewPropertyInfo.SerializedValueSize ||
+				FMemory::Memcmp(OldSerializedValuePtr, NewSerializedValuePtr, OldPropertyInfo->SerializedValueSize) != 0)
 			{
-				uint8* OldSerializedValuePtr = OriginalCDOProperties.Bytes.GetData() + OldPropertyInfo->SerializedValueOffset;
-				uint8* NewSerializedValuePtr = ReconstructedCDOProperties.Bytes.GetData() + NewPropertyInfo.SerializedValueOffset;
-				if (OldPropertyInfo->SerializedValueSize != NewPropertyInfo.SerializedValueSize ||
-					FMemory::Memcmp(OldSerializedValuePtr, NewSerializedValuePtr, OldPropertyInfo->SerializedValueSize) != 0)
+				// Property value has changed so add it to the list of properties that need updating on instances
+				FPropertyToUpdate PropertyToUpdate;
+				PropertyToUpdate.Property = NewPropertyInfo.Property;
+				PropertyToUpdate.NewValuePtr = nullptr;
+				PropertyToUpdate.SubobjectName = NewPropertyInfo.SubobjectName;
+
+				if (NewPropertyInfo.Property->GetOuter() == NewClass)
 				{
-					// Property value has changed so add it to the list of properties that need updating on instances
-					FPropertyToUpdate PropertyToUpdate;
-					PropertyToUpdate.Property = NewPropertyInfo.Property;
-					PropertyToUpdate.OldSerializedValuePtr = OldSerializedValuePtr;
 					PropertyToUpdate.NewValuePtr = PropertyToUpdate.Property->ContainerPtrToValuePtr<uint8>(NewClass->GetDefaultObject());
+				}
+				else if (NewPropertyInfo.SubobjectName != NAME_None)
+				{
+					UObject* DefaultSubobjectPtr = FindDefaultSubobject(DefaultSubobjectArray, NewPropertyInfo.SubobjectName);
+					if (DefaultSubobjectPtr && NewPropertyInfo.Property->GetOuter() == DefaultSubobjectPtr->GetClass())
+					{
+						PropertyToUpdate.NewValuePtr = PropertyToUpdate.Property->ContainerPtrToValuePtr<uint8>(DefaultSubobjectPtr);
+					}
+				}
+				if (PropertyToUpdate.NewValuePtr)
+				{					
+					PropertyToUpdate.OldSerializedValuePtr = OldSerializedValuePtr;
 					PropertyToUpdate.OldSerializedSize = OldPropertyInfo->SerializedValueSize;
+
 					PropertiesToUpdate.Add(PropertyToUpdate);
-				}
+				}					
 			}
 		}
 	}
@@ -312,21 +354,39 @@ void FHotReloadClassReinstancer::UpdateDefaultProperties()
 		for (FObjectIterator It(NewClass); It; ++It)
 		{
 			UObject* ObjectPtr = *It;
+			DefaultSubobjectArray.Empty(DefaultSubobjectArrayCapacity);
+			ObjectPtr->CollectDefaultSubobjects(DefaultSubobjectArray);
+
 			for (auto& PropertyToUpdate : PropertiesToUpdate)
 			{
-				uint8* InstanceValuePtr = PropertyToUpdate.Property->ContainerPtrToValuePtr<uint8>(ObjectPtr);
-
-				// Serialize current value to a byte array as we don't have the previous CDO to compare against, we only have its serialized property data
-				CurrentValueSerializedData.Empty(CurrentValueSerializedData.Num() + CurrentValueSerializedData.GetSlack());
-				FMemoryWriter CurrentValueWriter(CurrentValueSerializedData);
-				PropertyToUpdate.Property->SerializeItem(CurrentValueWriter, InstanceValuePtr);
-				
-				// Update only when the current value on the instance is identical to the original CDO
-				if (CurrentValueSerializedData.Num() == PropertyToUpdate.OldSerializedSize &&
-					FMemory::Memcmp(CurrentValueSerializedData.GetData(), PropertyToUpdate.OldSerializedValuePtr, CurrentValueSerializedData.Num()) == 0)
+				uint8* InstanceValuePtr = nullptr;
+				if (PropertyToUpdate.SubobjectName == NAME_None)
+				{
+					InstanceValuePtr = PropertyToUpdate.Property->ContainerPtrToValuePtr<uint8>(ObjectPtr);
+				}
+				else
 				{
-					// Update with the new value
-					PropertyToUpdate.Property->CopyCompleteValue(InstanceValuePtr, PropertyToUpdate.NewValuePtr);
+					UObject* DefaultSubobjectPtr = FindDefaultSubobject(DefaultSubobjectArray, PropertyToUpdate.SubobjectName);
+					if (DefaultSubobjectPtr && PropertyToUpdate.Property->GetOuter() == DefaultSubobjectPtr->GetClass())
+					{
+						InstanceValuePtr = PropertyToUpdate.Property->ContainerPtrToValuePtr<uint8>(DefaultSubobjectPtr);
+					}
+				}
+
+				if (InstanceValuePtr)
+				{
+					// Serialize current value to a byte array as we don't have the previous CDO to compare against, we only have its serialized property data
+					CurrentValueSerializedData.Empty(CurrentValueSerializedData.Num() + CurrentValueSerializedData.GetSlack());
+					FMemoryWriter CurrentValueWriter(CurrentValueSerializedData);
+					PropertyToUpdate.Property->SerializeItem(CurrentValueWriter, InstanceValuePtr);
+
+					// Update only when the current value on the instance is identical to the original CDO
+					if (CurrentValueSerializedData.Num() == PropertyToUpdate.OldSerializedSize &&
+						FMemory::Memcmp(CurrentValueSerializedData.GetData(), PropertyToUpdate.OldSerializedValuePtr, CurrentValueSerializedData.Num()) == 0)
+					{
+						// Update with the new value
+						PropertyToUpdate.Property->CopyCompleteValue(InstanceValuePtr, PropertyToUpdate.NewValuePtr);
+					}
 				}
 			}
 		}
diff --git a/Engine/Source/Developer/HotReload/Private/HotReloadClassReinstancer.h b/Engine/Source/Developer/HotReload/Private/HotReloadClassReinstancer.h
--- a/Engine/Source/Developer/HotReload/Private/HotReloadClassReinstancer.h
+++ b/Engine/Source/Developer/HotReload/Private/HotReloadClassReinstancer.h
@@ -15,11 +15,13 @@ class FHotReloadClassReinstancer : public FBlueprintCompileReinstancer
 	{
 		FCDOProperty()
 			: Property(nullptr)
+			, SubobjectName(NAME_None)
 			, SerializedValueOffset(0)
 			, SerializedValueSize(0)
 		{}
 
 		UProperty* Property;
+		FName SubobjectName;
 		int64 SerializedValueOffset;
 		int64 SerializedValueSize;
 	};
diff --git a/Engine/Source/Runtime/CoreUObject/Private/UObject/UObjectGlobals.cpp b/Engine/Source/Runtime/CoreUObject/Private/UObject/UObjectGlobals.cpp
--- a/Engine/Source/Runtime/CoreUObject/Private/UObject/UObjectGlobals.cpp
+++ b/Engine/Source/Runtime/CoreUObject/Private/UObject/UObjectGlobals.cpp
@@ -2224,9 +2224,15 @@ UObject* StaticConstructObject
 		(
 			!InTemplate || 
 			(InName != NAME_None && InTemplate == UObject::GetArchetypeFromRequiredInfo(InClass, InOuter, InName, InFlags))
-		);
-	bool bRecycledSubobject = false;
-	Result = StaticAllocateObject(InClass, InOuter, InName, InFlags, bIsNativeFromCDO, &bRecycledSubobject);
+		);	
+#if WITH_HOT_RELOAD
+	// Do not recycle subobjects when performing hot-reload as they may contain old property values.
+	const bool bCanRecycleSubobjects = bIsNativeFromCDO && !GIsHotReload;
+#else
+	const bool bCanRecycleSubobjects = bIsNativeFromCDO;
+#endif
+	bool bRecycledSubobject = false;	
+	Result = StaticAllocateObject(InClass, InOuter, InName, InFlags, bCanRecycleSubobjects, &bRecycledSubobject);
 	check(Result != NULL);
 	// Don't call the constructor on recycled subobjects, they haven't been destroyed.
 	if (!bRecycledSubobject)
