[INTEGRATE] Change 2417928 by [REDACTED_PERSON][REDACTED_HANDLE] on 2015/01/24 15:53:14

	Fixed failed C++ hot reload not displaying an error to the user (UE-8241)
	- Also added slightly better progress reporting to hot reload when adding code to a project

diff --git a/Engine/Source/Developer/HotReload/Private/HotReload.cpp b/Engine/Source/Developer/HotReload/Private/HotReload.cpp
--- a/Engine/Source/Developer/HotReload/Private/HotReload.cpp
+++ b/Engine/Source/Developer/HotReload/Private/HotReload.cpp
@@ -72,7 +72,7 @@ public:
 	virtual bool IsCurrentlyCompiling() const override { return ModuleCompileProcessHandle.IsValid(); }
 	virtual void RequestStopCompilation() override { bRequestCancelCompilation = true; }
 	virtual void AddHotReloadFunctionRemap(Native NewFunctionPointer, Native OldFunctionPointer) override;	
-	virtual void RebindPackages(TArray< UPackage* > Packages, TArray< FName > DependentModules, const bool bWaitForCompletion, FOutputDevice &Ar) override;
+	virtual ECompilationResult::Type RebindPackages(TArray< UPackage* > Packages, TArray< FName > DependentModules, const bool bWaitForCompletion, FOutputDevice &Ar) override;
 	virtual ECompilationResult::Type DoHotReloadFromEditor() override;
 	virtual FHotReloadEvent& OnHotReload() override { return HotReloadEvent; }	
 	virtual FModuleCompilerStartedEvent& OnModuleCompilerStarted() override { return ModuleCompilerStartedEvent; }
@@ -739,7 +739,7 @@ ECompilationResult::Type FHotReloadModule::DoHotReloadInternal(bool bRecompileFi
 	return Result;
 }
 
-void FHotReloadModule::RebindPackages(TArray<UPackage*> InPackages, TArray<FName> DependentModules, const bool bWaitForCompletion, FOutputDevice &Ar)
+ECompilationResult::Type FHotReloadModule::RebindPackages(TArray<UPackage*> InPackages, TArray<FName> DependentModules, const bool bWaitForCompletion, FOutputDevice &Ar)
 {
 	ECompilationResult::Type Result = ECompilationResult::Unknown;
 	double Duration = 0.0;
@@ -748,6 +748,8 @@ void FHotReloadModule::RebindPackages(TArray<UPackage*> InPackages, TArray<FName
 		Result = RebindPackagesInternal(InPackages, DependentModules, bWaitForCompletion, Ar);
 	}
 	RecordAnalyticsEvent(TEXT("Rebind"), Result, Duration, InPackages.Num(), DependentModules.Num());
+
+	return Result;
 }
 
 ECompilationResult::Type FHotReloadModule::RebindPackagesInternal(TArray<UPackage*> InPackages, TArray<FName> DependentModules, const bool bWaitForCompletion, FOutputDevice &Ar)
diff --git a/Engine/Source/Developer/ModuleUI/Private/SModuleUI.cpp b/Engine/Source/Developer/ModuleUI/Private/SModuleUI.cpp
--- a/Engine/Source/Developer/ModuleUI/Private/SModuleUI.cpp
+++ b/Engine/Source/Developer/ModuleUI/Private/SModuleUI.cpp
@@ -257,7 +257,7 @@ FReply SModuleUI::FModuleListItem::OnRecompileClicked()
 		{
 			// Perform a hot reload
 			const bool bWaitForCompletion = true;			
-			HotReloadSupport.RebindPackages(PackagesToRebind, TArray<FName>(), bWaitForCompletion, *GLog);
+			ECompilationResult::Type Result = HotReloadSupport.RebindPackages(PackagesToRebind, TArray<FName>(), bWaitForCompletion, *GLog);
 		}
 		else
 		{
@@ -265,7 +265,7 @@ FReply SModuleUI::FModuleListItem::OnRecompileClicked()
 			const bool bReloadAfterRecompile = true;
 			const bool bForceCodeProject = false;
 			const bool bFailIfGeneratedCodeChanges = true;
-			HotReloadSupport.RecompileModule(ModuleName, bReloadAfterRecompile, *GLog, bFailIfGeneratedCodeChanges, bForceCodeProject);
+			const bool bRecompileSucceeded = HotReloadSupport.RecompileModule(ModuleName, bReloadAfterRecompile, *GLog, bFailIfGeneratedCodeChanges, bForceCodeProject);
 		}
 	}
 
diff --git a/Engine/Source/Editor/GameProjectGeneration/Private/GameProjectUtils.cpp b/Engine/Source/Editor/GameProjectGeneration/Private/GameProjectUtils.cpp
--- a/Engine/Source/Editor/GameProjectGeneration/Private/GameProjectUtils.cpp
+++ b/Engine/Source/Editor/GameProjectGeneration/Private/GameProjectUtils.cpp
@@ -2956,7 +2956,7 @@ bool GameProjectUtils::AddCodeToProject_Internal(const FString& NewClassName, co
 		return false;
 	}
 
-	FScopedSlowTask SlowTask( 5, LOCTEXT( "AddingCodeToProject", "Adding code to project..." ) );
+	FScopedSlowTask SlowTask( 6, LOCTEXT( "AddingCodeToProject", "Adding code to project..." ) );
 	SlowTask.MakeDialog();
 
 	SlowTask.EnterProgressFrame();
@@ -3042,6 +3042,8 @@ bool GameProjectUtils::AddCodeToProject_Internal(const FString& NewClassName, co
 		SourceControlProvider.Execute(ISourceControlOperation::Create<FMarkForAdd>(), FilesToCheckOut);
 	}
 
+	SlowTask.EnterProgressFrame( 1.0f, LOCTEXT("CompilingCPlusPlusCode", "Compiling new C++ code.  Please wait..."));
+
 	OutHeaderFilePath = NewHeaderFilename;
 	OutCppFilePath = NewCppFilename;
 
@@ -3092,7 +3094,12 @@ bool GameProjectUtils::AddCodeToProject_Internal(const FString& NewClassName, co
 			{
 				// Perform a hot reload
 				const bool bWaitForCompletion = true;			
-				HotReloadSupport.RebindPackages(PackagesToRebind, TArray<FName>(), bWaitForCompletion, *GWarn);
+				ECompilationResult::Type CompilationResult = HotReloadSupport.RebindPackages( PackagesToRebind, TArray<FName>(), bWaitForCompletion, *GWarn );
+				if( CompilationResult != ECompilationResult::Succeeded && CompilationResult != ECompilationResult::UpToDate )
+				{
+					OutFailReason = FText::Format(LOCTEXT("FailedToHotReloadModuleFmt", "Failed to automatically hot reload the '{0}' module."), FText::FromString(ModuleInfo.ModuleName));
+					return false;
+				}
 			}
 			else
 			{
diff --git a/Engine/Source/Runtime/CoreUObject/Public/Misc/HotReloadInterface.h b/Engine/Source/Runtime/CoreUObject/Public/Misc/HotReloadInterface.h
--- a/Engine/Source/Runtime/CoreUObject/Public/Misc/HotReloadInterface.h
+++ b/Engine/Source/Runtime/CoreUObject/Public/Misc/HotReloadInterface.h
@@ -75,8 +75,10 @@ public:
 	* @param	DependentModules	Additional modules that don't contain UObjects, but rely on them
 	* @param	bWaitForCompletion	True if RebindPackages should not return until the recompile and reload has completed
 	* @param	Ar					Output device for logging compilation status
+	* 
+	* @return	If bWaitForCompletion was set to true, this will return the result of the compilation, otherwise will return ECompilationResult::Unknown
 	*/
-	virtual void RebindPackages(TArray<UPackage*> Packages, TArray<FName> DependentModules, const bool bWaitForCompletion, FOutputDevice &Ar) = 0;
+	virtual ECompilationResult::Type RebindPackages(TArray<UPackage*> Packages, TArray<FName> DependentModules, const bool bWaitForCompletion, FOutputDevice &Ar) = 0;
 
 	/** Called when a Hot Reload event has completed. 
 	 * 
diff --git a/Engine/Source/Runtime/Engine/Private/UnrealEngine.cpp b/Engine/Source/Runtime/Engine/Private/UnrealEngine.cpp
--- a/Engine/Source/Runtime/Engine/Private/UnrealEngine.cpp
+++ b/Engine/Source/Runtime/Engine/Private/UnrealEngine.cpp
@@ -2676,7 +2676,7 @@ bool UEngine::HandleHotReloadCommand( const TCHAR* Cmd, FOutputDevice& Ar )
 		PackagesToRebind.Add( Package );
 		const bool bWaitForCompletion = true;	// Always wait when hotreload is initiated from the console
 		IHotReloadInterface& HotReloadSupport = FModuleManager::LoadModuleChecked<IHotReloadInterface>("HotReload");
-		HotReloadSupport.RebindPackages(PackagesToRebind, TArray<FName>(), bWaitForCompletion, Ar);
+		const ECompilationResult::Type CompilationResult = HotReloadSupport.RebindPackages(PackagesToRebind, TArray<FName>(), bWaitForCompletion, Ar);
 	}
 	return true;
 }
