[INTEGRATE] Change 2421953 by [REDACTED_PERSON][REDACTED_HANDLE].7 on 2015/01/28 06:48:44

	Fixed hot reload being able to add two scoped slow tasks to a single progress frame

diff --git a/Engine/Source/Developer/HotReload/Private/HotReload.cpp b/Engine/Source/Developer/HotReload/Private/HotReload.cpp
--- a/Engine/Source/Developer/HotReload/Private/HotReload.cpp
+++ b/Engine/Source/Developer/HotReload/Private/HotReload.cpp
@@ -476,7 +476,7 @@ bool FHotReloadModule::RecompileModule(const FName InModuleName, const bool bRel
 	Args.Add( TEXT("CodeModuleName"), FText::FromName( InModuleName ) );
 	const FText StatusUpdate = FText::Format( NSLOCTEXT("ModuleManager", "Recompile_SlowTaskName", "Compiling {CodeModuleName}..."), Args );
 
-	FScopedSlowTask SlowTask(1, StatusUpdate);
+	FScopedSlowTask SlowTask(2, StatusUpdate);
 	SlowTask.MakeDialog();
 
 	ModuleCompilerStartedEvent.Broadcast();
@@ -489,6 +489,8 @@ bool FHotReloadModule::RecompileModule(const FName InModuleName, const bool bRel
 	const bool bWasModuleLoaded = FModuleManager::Get().IsModuleLoaded( InModuleName );
 	const bool bUseRollingModuleNames = bWasModuleLoaded;
 
+	SlowTask.EnterProgressFrame();
+
 	bool bWasSuccessful = true;
 	if( bUseRollingModuleNames )
 	{
@@ -510,6 +512,8 @@ bool FHotReloadModule::RecompileModule(const FName InModuleName, const bool bRel
 		bWasSuccessful = RecompileModuleDLLs(ModulesToRecompile, Ar, bFailIfGeneratedCodeChanges, bForceCodeProject);		
 	}
 
+	SlowTask.EnterProgressFrame();
+	
 	if( bWasSuccessful )
 	{
 		// Shutdown the module if it's already running
