PR #783: Remove duplicate pipe() call on Linux

- Contributed by [REDACTED_PERSON].

https://github.com/EpicGames/UnrealEngine/pull/783

diff --git a/Engine/Source/Developer/HotReload/Private/HotReload.cpp b/Engine/Source/Developer/HotReload/Private/HotReload.cpp
--- a/Engine/Source/Developer/HotReload/Private/HotReload.cpp
+++ b/Engine/Source/Developer/HotReload/Private/HotReload.cpp
@@ -1278,15 +1278,8 @@ bool FHotReloadModule::InvokeUnrealBuildToolForCompile(const FString& InCmdLineP
 	check(!IsCurrentlyCompiling());
 
 	// Setup output redirection pipes, so that we can harvest compiler output and display it ourselves
-#if PLATFORM_LINUX
-	int pipefd[2];
-	pipe(pipefd);
-	void* PipeRead = &pipefd[0];
-	void* PipeWrite = &pipefd[1];
-#else
 	void* PipeRead = NULL;
 	void* PipeWrite = NULL;
-#endif
 
 	verify(FPlatformProcess::CreatePipe(PipeRead, PipeWrite));
 	ModuleCompileReadPipeText = TEXT("");
@@ -1295,11 +1288,7 @@ bool FHotReloadModule::InvokeUnrealBuildToolForCompile(const FString& InCmdLineP
 
 	// We no longer need the Write pipe so close it.
 	// We DO need the Read pipe however...
-#if PLATFORM_LINUX
-	close(*(int*)PipeWrite);
-#else
 	FPlatformProcess::ClosePipe(0, PipeWrite);
-#endif
 
 	if (!ProcHandle.IsValid())
 	{
@@ -1414,11 +1403,7 @@ void FHotReloadModule::CheckForFinishedModuleDLLCompile(const bool bWaitForCompl
 			ModuleCompileProcessHandle.Close();
 			ModuleCompileProcessHandle.Reset();
 
-#if PLATFORM_LINUX
-			close(*(int *)ModuleCompileReadPipe);
-#else
 			FPlatformProcess::ClosePipe(ModuleCompileReadPipe, 0);
-#endif
 
 			Ar.Log(*ModuleCompileReadPipeText);
 			const FString FinalOutput = ModuleCompileReadPipeText;
