Fix for hot-reload not always detecting DSO changes in the constructor

diff --git a/Engine/Source/Developer/HotReload/Private/HotReloadClassReinstancer.cpp b/Engine/Source/Developer/HotReload/Private/HotReloadClassReinstancer.cpp
--- a/Engine/Source/Developer/HotReload/Private/HotReloadClassReinstancer.cpp
+++ b/Engine/Source/Developer/HotReload/Private/HotReloadClassReinstancer.cpp
@@ -117,6 +117,7 @@ void FHotReloadClassReinstancer::SerializeCDOProperties(UObject* InObject, FHotR
 					{
 						// Serialize all DSO properties too					
 						FCDOWriter DefaultSubobjectWriter(PropertyData, InObj, VisitedObjects, InObj->GetFName());
+						Seek(PropertyData.Bytes.Num());
 					}
 				}
 			}
