Fixing crash when performing hot-reload

diff --git a/Engine/Source/Developer/HotReload/Private/HotReloadClassReinstancer.cpp b/Engine/Source/Developer/HotReload/Private/HotReloadClassReinstancer.cpp
--- a/Engine/Source/Developer/HotReload/Private/HotReloadClassReinstancer.cpp
+++ b/Engine/Source/Developer/HotReload/Private/HotReloadClassReinstancer.cpp
@@ -23,6 +23,8 @@ void FHotReloadClassReinstancer::SetupNewClassReinstancing(UClass* InNewClass, U
 
 	SaveClassFieldMapping(InOldClass);
 
+	ObjectsThatShouldUseOldStuff.Add(InOldClass); //CDO of REINST_ class can be used as archetype
+
 	TArray<UClass*> ChildrenOfClass;
 	GetDerivedClasses(InOldClass, ChildrenOfClass);
 	for (auto ClassIt = ChildrenOfClass.CreateConstIterator(); ClassIt; ++ClassIt)
