Fix for hot-reload not working with Blueprints.

diff --git a/Engine/Source/Developer/HotReload/Private/HotReloadClassReinstancer.cpp b/Engine/Source/Developer/HotReload/Private/HotReloadClassReinstancer.cpp
--- a/Engine/Source/Developer/HotReload/Private/HotReloadClassReinstancer.cpp
+++ b/Engine/Source/Developer/HotReload/Private/HotReloadClassReinstancer.cpp
@@ -42,6 +42,10 @@ void FHotReloadClassReinstancer::SetupNewClassReinstancing(UClass* InNewClass, U
 				}
 
 				Children.AddUnique(ChildBP);
+				if (ChildBP->ParentClass == InOldClass)
+				{
+					ChildBP->ParentClass = NewClass;
+				}
 			}
 			else
 			{
