Fixing another hot-reload crash: this time when removing code responsible for creating default subobject from the constructor without removing its property from the parent class.

diff --git a/Engine/Source/Developer/HotReload/Private/HotReloadClassReinstancer.cpp b/Engine/Source/Developer/HotReload/Private/HotReloadClassReinstancer.cpp
--- a/Engine/Source/Developer/HotReload/Private/HotReloadClassReinstancer.cpp
+++ b/Engine/Source/Developer/HotReload/Private/HotReloadClassReinstancer.cpp
@@ -304,6 +304,32 @@ void FHotReloadClassReinstancer::UpdateDefaultProperties()
 		uint8* NewValuePtr;
 		int64 OldSerializedSize;
 	};
+	/** Memory writer archive that supports UObject values the same way as FCDOWriter. */
+	class FPropertyValueMemoryWriter : public FMemoryWriter
+	{
+	public:
+		FPropertyValueMemoryWriter(TArray<uint8>& OutData)
+			: FMemoryWriter(OutData)
+		{}
+		virtual FArchive& operator<<(class UObject*& InObj) override
+		{
+			FArchive& Ar = *this;
+			if (InObj)
+			{
+				FName ClassName = InObj->GetClass()->GetFName();
+				FName ObjectName = InObj->GetFName();
+				Ar << ClassName;
+				Ar << ObjectName;
+			}
+			else
+			{
+				FName UnusedName = NAME_None;
+				Ar << UnusedName;
+				Ar << UnusedName;
+			}
+			return *this;
+		}
+	};
 
 	// Collect default subobjects to update their properties too
 	const int32 DefaultSubobjectArrayCapacity = 16;
@@ -384,7 +410,7 @@ void FHotReloadClassReinstancer::UpdateDefaultProperties()
 				{
 					// Serialize current value to a byte array as we don't have the previous CDO to compare against, we only have its serialized property data
 					CurrentValueSerializedData.Empty(CurrentValueSerializedData.Num() + CurrentValueSerializedData.GetSlack());
-					FMemoryWriter CurrentValueWriter(CurrentValueSerializedData);
+					FPropertyValueMemoryWriter CurrentValueWriter(CurrentValueSerializedData);
 					PropertyToUpdate.Property->SerializeItem(CurrentValueWriter, InstanceValuePtr);
 
 					// Update only when the current value on the instance is identical to the original CDO
