[INTEGRATE] Change 2429033 by [REDACTED_PERSON][REDACTED_HANDLE] on 2015/02/02 19:12:32

	Made sure that FHotReloadModule still reported slow task progress while stuck in a busy loop

diff --git a/Engine/Source/Developer/HotReload/Private/HotReload.cpp b/Engine/Source/Developer/HotReload/Private/HotReload.cpp
--- a/Engine/Source/Developer/HotReload/Private/HotReload.cpp
+++ b/Engine/Source/Developer/HotReload/Private/HotReload.cpp
@@ -1364,6 +1364,8 @@ void FHotReloadModule::CheckForFinishedModuleDLLCompile(const bool bWaitForCompl
 					break;
 				}
 
+				SlowTask.EnterProgressFrame(0.0f);
+
 				// Give up a small timeslice if we haven't finished recompiling yet
 				FPlatformProcess::Sleep( 0.01f );
 			}
