Fixing hot-reload from IDE working only the first time. Added code to strip module suffix when checking module filenames.

diff --git a/Engine/Source/Developer/HotReload/Private/HotReload.cpp b/Engine/Source/Developer/HotReload/Private/HotReload.cpp
--- a/Engine/Source/Developer/HotReload/Private/HotReload.cpp
+++ b/Engine/Source/Developer/HotReload/Private/HotReload.cpp
@@ -183,6 +183,11 @@ private:
 	 */
 	void OnHotReloadBinariesChanged(const TArray<struct FFileChangeData>& FileChanges);
 
+	/**
+	 * Strips hot-reload suffix from module filename.
+	 */
+	static void StripModuleSuffixFromFilename(FString& InOutModuleFilename, const FString& ModuleName);
+
 	/**
 	 * Broadcasts that a hot reload just finished. 
 	 *
@@ -900,16 +905,20 @@ void FHotReloadModule::OnHotReloadBinariesChanged(const TArray<struct FFileChang
 				const FString Filename = FPaths::GetCleanFilename(Change.Filename);
 				if (Filename.EndsWith(FPlatformProcess::GetModuleExtension()))
 				{
-					for( int32 GameModuleIndex = 0; GameModuleIndex < GameModuleNames.Num(); ++GameModuleIndex )
+					for (int32 GameModuleIndex = 0; GameModuleIndex < GameModuleNames.Num(); ++GameModuleIndex)
 					{
-						const FString& GameModuleName = GameModuleNames[ GameModuleIndex ];
-						const FString& GameModuleFilePath = GameModuleFilePaths[ GameModuleIndex ];
+						const FString& GameModuleName = GameModuleNames[GameModuleIndex];
+						const FString& GameModuleFilePath = GameModuleFilePaths[GameModuleIndex];
+
+						// Handle module files which have already been hot-reloaded.
+						FString GameModuleFileNameWithoutExtension = FPaths::GetBaseFilename(GameModuleFilePath);
+						StripModuleSuffixFromFilename(GameModuleFileNameWithoutExtension, GameModuleName);
 
-						const FString GameModuleFileNameWithoutExtension = FPaths::GetBaseFilename( GameModuleFilePath );
-						if( Filename.StartsWith( GameModuleFileNameWithoutExtension + TEXT( "-" ) ) )	// Hot reload always adds a numbered suffix preceded by a hyphen, but otherwise the module name must match exactly!
+						// Hot reload always adds a numbered suffix preceded by a hyphen, but otherwise the module name must match exactly!
+						if (Filename.StartsWith(GameModuleFileNameWithoutExtension + TEXT("-")))
 						{
-							if ( !NewModules.ContainsByPredicate([&](const FRecompiledModule& Module){ return Module.Name == GameModuleName; }) &&
-								 !ModulesRecentlyCompiledInTheEditor.Contains(FPaths::ConvertRelativePathToFull(Change.Filename)))
+							if (!NewModules.ContainsByPredicate([&](const FRecompiledModule& Module){ return Module.Name == GameModuleName; }) &&
+								!ModulesRecentlyCompiledInTheEditor.Contains(FPaths::ConvertRelativePathToFull(Change.Filename)))
 							{
 								// Add to queue. We do not hot-reload here as there may potentially be other modules being compiled.
 								NewModules.Add(FRecompiledModule(GameModuleName, Change.Filename));
@@ -923,6 +932,31 @@ void FHotReloadModule::OnHotReloadBinariesChanged(const TArray<struct FFileChang
 	}
 }
 
+void FHotReloadModule::StripModuleSuffixFromFilename(FString& InOutModuleFilename, const FString& ModuleName)
+{
+	// First hyphen is where the UE4Edtior prefix ends
+	int32 FirstHyphenIndex = INDEX_NONE;
+	if (InOutModuleFilename.FindChar('-', FirstHyphenIndex))
+	{
+		// Second hyphen means we already have a hot-reloaded module or other than Development config module
+		int32 SecondHyphenIndex = FirstHyphenIndex;
+		do
+		{
+			SecondHyphenIndex = InOutModuleFilename.Find(TEXT("-"), ESearchCase::IgnoreCase, ESearchDir::FromStart, SecondHyphenIndex + 1);
+			if (SecondHyphenIndex != INDEX_NONE)
+			{
+				// Make sure that the section between hyphens is the expected module name. This guards against cases where module name has a hyphen inside.
+				FString HotReloadedModuleName = InOutModuleFilename.Mid(FirstHyphenIndex + 1, SecondHyphenIndex - FirstHyphenIndex - 1);
+				if (HotReloadedModuleName == ModuleName)
+				{
+					InOutModuleFilename = InOutModuleFilename.Mid(0, SecondHyphenIndex);
+					SecondHyphenIndex = INDEX_NONE;
+				}
+			}
+		} while (SecondHyphenIndex != INDEX_NONE);
+	}
+}
+
 void FHotReloadModule::InitHotReloadWatcher()
 {
 	FDirectoryWatcherModule& DirectoryWatcherModule = FModuleManager::Get().LoadModuleChecked<FDirectoryWatcherModule>(TEXT("DirectoryWatcher"));
