[INTEGRATE] Change 2430066 by [REDACTED_PERSON][REDACTED_HANDLE] on 2015/02/03 11:54:22

	Stopped the "Compiling C++ Code" message in the editor MainFrame always having a cancel button

	FModuleCompilerStartedEvent now passes a bool to say whether the hot-reload/compile operation is being run async, or whether it's blocking. We now only show the cancel button for async compiles, as the UI is otherwise blocked waiting for the compile to finish.

	Added a bWaitForCompletion input parameter to DoHotReloadFromEditor. This mirrors the RebindPackages API that DoHotReloadFromEditor internally calls, and prevents UPackFactory needing to perform a manual tick.

diff --git a/Engine/Plugins/Developer/VisualStudioSourceCodeAccess/Source/VisualStudioSourceCodeAccess/Private/VisualStudioSourceCodeAccessor.cpp b/Engine/Plugins/Developer/VisualStudioSourceCodeAccess/Source/VisualStudioSourceCodeAccess/Private/VisualStudioSourceCodeAccessor.cpp
--- a/Engine/Plugins/Developer/VisualStudioSourceCodeAccess/Source/VisualStudioSourceCodeAccess/Private/VisualStudioSourceCodeAccessor.cpp
+++ b/Engine/Plugins/Developer/VisualStudioSourceCodeAccess/Source/VisualStudioSourceCodeAccess/Private/VisualStudioSourceCodeAccessor.cpp
@@ -61,7 +61,7 @@ enum class EAccessVisualStudioResult : uint8
 };
 
 /** save all open documents in visual studio, when recompiling */
-void SaveVisualStudioDocuments()
+void OnModuleCompileStarted(bool bIsAsyncCompile)
 {
 	FVisualStudioSourceCodeAccessModule& VisualStudioSourceCodeAccessModule = FModuleManager::LoadModuleChecked<FVisualStudioSourceCodeAccessModule>(TEXT("VisualStudioSourceCodeAccess"));
 	VisualStudioSourceCodeAccessModule.GetAccessor().SaveAllOpenDocuments();
@@ -73,7 +73,7 @@ void FVisualStudioSourceCodeAccessor::Startup()
 
 #if WITH_EDITOR
 	// Setup compilation for saving all VS documents upon compilation start
-	SaveVisualStudioDocumentsDelegateHandle = IHotReloadModule::Get().OnModuleCompilerStarted().AddStatic( &SaveVisualStudioDocuments );
+	SaveVisualStudioDocumentsDelegateHandle = IHotReloadModule::Get().OnModuleCompilerStarted().AddStatic( &OnModuleCompileStarted );
 #endif
 
 	// Cache this so we don't have to do it on a background thread
diff --git a/Engine/Source/Developer/HotReload/Private/HotReload.cpp b/Engine/Source/Developer/HotReload/Private/HotReload.cpp
--- a/Engine/Source/Developer/HotReload/Private/HotReload.cpp
+++ b/Engine/Source/Developer/HotReload/Private/HotReload.cpp
@@ -73,7 +73,7 @@ public:
 	virtual void RequestStopCompilation() override { bRequestCancelCompilation = true; }
 	virtual void AddHotReloadFunctionRemap(Native NewFunctionPointer, Native OldFunctionPointer) override;	
 	virtual ECompilationResult::Type RebindPackages(TArray< UPackage* > Packages, TArray< FName > DependentModules, const bool bWaitForCompletion, FOutputDevice &Ar) override;
-	virtual ECompilationResult::Type DoHotReloadFromEditor() override;
+	virtual ECompilationResult::Type DoHotReloadFromEditor(const bool bWaitForCompletion) override;
 	virtual FHotReloadEvent& OnHotReload() override { return HotReloadEvent; }	
 	virtual FModuleCompilerStartedEvent& OnModuleCompilerStarted() override { return ModuleCompilerStartedEvent; }
 	virtual FModuleCompilerFinishedEvent& OnModuleCompilerFinished() override { return ModuleCompilerFinishedEvent; }
@@ -484,7 +484,7 @@ bool FHotReloadModule::RecompileModule(const FName InModuleName, const bool bRel
 	FScopedSlowTask SlowTask(2, StatusUpdate);
 	SlowTask.MakeDialog();
 
-	ModuleCompilerStartedEvent.Broadcast();
+	ModuleCompilerStartedEvent.Broadcast(false); // we never perform an async compile
 
 	// Update our set of known modules, in case we don't already know about this module
 	FModuleManager::Get().AddModule( InModuleName );
@@ -587,14 +587,11 @@ void FHotReloadModule::AddHotReloadFunctionRemap(Native NewFunctionPointer, Nati
 	HotReloadFunctionRemap.Add(OldFunctionPointer, NewFunctionPointer);
 }
 
-ECompilationResult::Type FHotReloadModule::DoHotReloadFromEditor()
+ECompilationResult::Type FHotReloadModule::DoHotReloadFromEditor(const bool bWaitForCompletion)
 {
 	// Get all game modules we want to compile
 	TArray<FString> GameModuleNames;
 	GetGameModules(GameModuleNames);
-	
-	// Don't wait -- we want compiling to happen asynchronously
-	const bool bWaitForCompletion = false;
 
 	TArray<UPackage*> PackagesToRebind;
 	TArray<FName> DependentModules;
@@ -1088,7 +1085,7 @@ bool FHotReloadModule::RecompileModulesAsync( const TArray< FName > ModuleNames,
 	// NOTE: This method of recompiling always using a rolling file name scheme, since we never want to unload before
 	// we start recompiling, and we need the output DLL to be unlocked before we invoke the compiler
 
-	ModuleCompilerStartedEvent.Broadcast();
+	ModuleCompilerStartedEvent.Broadcast(!bWaitForCompletion); // we perform an async compile providing we're not waiting for completion
 
 	TArray< FModuleToRecompile > ModulesToRecompile;
 
diff --git a/Engine/Source/Editor/LevelEditor/Private/LevelEditorActions.cpp b/Engine/Source/Editor/LevelEditor/Private/LevelEditorActions.cpp
--- a/Engine/Source/Editor/LevelEditor/Private/LevelEditorActions.cpp
+++ b/Engine/Source/Editor/LevelEditor/Private/LevelEditorActions.cpp
@@ -923,7 +923,9 @@ void FLevelEditorActionCallbacks::RecompileGameCode_Clicked()
 	IHotReloadInterface& HotReloadSupport = FModuleManager::LoadModuleChecked<IHotReloadInterface>(HotReloadModule);
 	if( !HotReloadSupport.IsCurrentlyCompiling() )
 	{
-		HotReloadSupport.DoHotReloadFromEditor();
+		// Don't wait -- we want compiling to happen asynchronously
+		const bool bWaitForCompletion = false;
+		HotReloadSupport.DoHotReloadFromEditor(bWaitForCompletion);
 	}
 }
 
diff --git a/Engine/Source/Editor/MainFrame/Private/MainFrameModule.cpp b/Engine/Source/Editor/MainFrame/Private/MainFrameModule.cpp
--- a/Engine/Source/Editor/MainFrame/Private/MainFrameModule.cpp
+++ b/Engine/Source/Editor/MainFrame/Private/MainFrameModule.cpp
@@ -716,7 +716,7 @@ bool FMainFrameModule::ShouldShowProjectDialogAtStartup( ) const
 /* FMainFrameModule event handlers
  *****************************************************************************/
 
-void FMainFrameModule::HandleLevelEditorModuleCompileStarted( )
+void FMainFrameModule::HandleLevelEditorModuleCompileStarted( bool bIsAsyncCompile )
 {
 	ModuleCompileStartTime = FPlatformTime::Seconds();
 
@@ -731,7 +731,12 @@ void FMainFrameModule::HandleLevelEditorModuleCompileStarted( )
 	Info.Image = FEditorStyle::GetBrush(TEXT("LevelEditor.RecompileGameCode"));
 	Info.ExpireDuration = 5.0f;
 	Info.bFireAndForget = false;
-	Info.ButtonDetails.Add(FNotificationButtonInfo(LOCTEXT("CancelC++Compilation", "Cancel"), FText(), FSimpleDelegate::CreateRaw(this, &FMainFrameModule::OnCancelCodeCompilationClicked)));
+	
+	// We can only show the cancel button on async builds
+	if (bIsAsyncCompile)
+	{
+		Info.ButtonDetails.Add(FNotificationButtonInfo(LOCTEXT("CancelC++Compilation", "Cancel"), FText(), FSimpleDelegate::CreateRaw(this, &FMainFrameModule::OnCancelCodeCompilationClicked)));
+	}
 
 	CompileNotificationPtr = FSlateNotificationManager::Get().AddNotification(Info);
 
diff --git a/Engine/Source/Editor/MainFrame/Private/MainFrameModule.h b/Engine/Source/Editor/MainFrame/Private/MainFrameModule.h
--- a/Engine/Source/Editor/MainFrame/Private/MainFrameModule.h
+++ b/Engine/Source/Editor/MainFrame/Private/MainFrameModule.h
@@ -130,7 +130,7 @@ public:
 private:
 
 	// Handles the level editor module starting to recompile.
-	void HandleLevelEditorModuleCompileStarted( );
+	void HandleLevelEditorModuleCompileStarted( bool bIsAsyncCompile );
 
 	// Handles the user requesting the current compilation to be canceled
 	void OnCancelCodeCompilationClicked();
diff --git a/Engine/Source/Editor/UnrealEd/Private/Factories/PackFactory.cpp b/Engine/Source/Editor/UnrealEd/Private/Factories/PackFactory.cpp
--- a/Engine/Source/Editor/UnrealEd/Private/Factories/PackFactory.cpp
+++ b/Engine/Source/Editor/UnrealEd/Private/Factories/PackFactory.cpp
@@ -496,19 +496,9 @@ UObject* UPackFactory::FactoryCreateBinary
 				{
 					// We can only hot-reload via DoHotReloadFromEditor when we already had code in our project
 					if (!HotReloadSupport.IsCurrentlyCompiling())
-					{					
-						HotReloadSupport.DoHotReloadFromEditor();
-					
-						UE_LOG(LogPackFactory, Warning, TEXT("Starting compile."));
-						// Wait until the compiler finishes
-						while (HotReloadSupport.IsCurrentlyCompiling())
-						{
-							// Tick slate while we wait
-							FSlateApplication::Get().Tick();						
-							// Tick hot reload - or nothing will happen !
-							HotReloadSupport.Tick();
-						}
-						UE_LOG(LogPackFactory, Warning, TEXT("Compile complete."));
+					{
+						const bool bWaitForCompletion = true;
+						HotReloadSupport.DoHotReloadFromEditor(bWaitForCompletion);
 					}
 				}
 				else
diff --git a/Engine/Source/Runtime/CoreUObject/Public/Misc/HotReloadInterface.h b/Engine/Source/Runtime/CoreUObject/Public/Misc/HotReloadInterface.h
--- a/Engine/Source/Runtime/CoreUObject/Public/Misc/HotReloadInterface.h
+++ b/Engine/Source/Runtime/CoreUObject/Public/Misc/HotReloadInterface.h
@@ -66,8 +66,10 @@ public:
 
 	/**
 	* Performs hot reload from the editor of all currently loaded game modules.
+	* @param	bWaitForCompletion	True if RebindPackages should not return until the recompile and reload has completed
+	* @return	If bWaitForCompletion was set to true, this will return the result of the compilation, otherwise will return ECompilationResult::Unknown
 	*/
-	virtual ECompilationResult::Type DoHotReloadFromEditor() = 0;
+	virtual ECompilationResult::Type DoHotReloadFromEditor(const bool bWaitForCompletion) = 0;
 
 	/**
 	* HotReload: Reloads the DLLs for given packages.
@@ -92,7 +94,7 @@ public:
 	 *
 	 * @return The event delegate.
 	 */
-	DECLARE_MULTICAST_DELEGATE(FModuleCompilerStartedEvent);
+	DECLARE_MULTICAST_DELEGATE_OneParam(FModuleCompilerStartedEvent, bool /*bIsAsyncCompile*/);
 	virtual FModuleCompilerStartedEvent& OnModuleCompilerStarted() = 0;
 
 	/**
