Fixing hot-reload crash

diff --git a/Engine/Source/Developer/HotReload/Private/HotReload.cpp b/Engine/Source/Developer/HotReload/Private/HotReload.cpp
--- a/Engine/Source/Developer/HotReload/Private/HotReload.cpp
+++ b/Engine/Source/Developer/HotReload/Private/HotReload.cpp
@@ -842,11 +842,11 @@ ECompilationResult::Type FHotReloadModule::RebindPackagesInternal(TArray<UPackag
 #if WITH_ENGINE
 void FHotReloadModule::ReinstanceClass(UClass* OldClass, UClass* NewClass)
 {	
-	FHotReloadClassReinstancer ReinstanceHelper(NewClass, OldClass);
-	if (ReinstanceHelper.ClassNeedsReinstancing())
+	auto ReinstanceHelper = FHotReloadClassReinstancer::Create(NewClass, OldClass);
+	if (ReinstanceHelper->ClassNeedsReinstancing())
 	{
 		UE_LOG(LogHotReload, Log, TEXT("Re-instancing %s after hot-reload."), NewClass ? *NewClass->GetName() : *OldClass->GetName());
-		ReinstanceHelper.ReinstanceObjectsAndUpdateDefaults();
+		ReinstanceHelper->ReinstanceObjectsAndUpdateDefaults();
 	}
 }
 #endif
diff --git a/Engine/Source/Developer/HotReload/Private/HotReloadClassReinstancer.h b/Engine/Source/Developer/HotReload/Private/HotReloadClassReinstancer.h
--- a/Engine/Source/Developer/HotReload/Private/HotReloadClassReinstancer.h
+++ b/Engine/Source/Developer/HotReload/Private/HotReloadClassReinstancer.h
@@ -101,6 +101,12 @@ public:
 
 	/** Reinstances all objects of the hot-reloaded class and update their properties to match the new CDO */
 	void ReinstanceObjectsAndUpdateDefaults();
+
+	/** Creates the reinstancer as a sharable object */
+	static TSharedPtr<FHotReloadClassReinstancer> Create(UClass* InNewClass, UClass* InOldClass)
+	{
+		return MakeShareable(new FHotReloadClassReinstancer(InNewClass, InOldClass));
+	}
 };
 
 #endif // WITH_ENGINE
\ No newline at end of file
diff --git a/Engine/Source/Editor/UnrealEd/Public/Kismet2/KismetReinstanceUtilities.h b/Engine/Source/Editor/UnrealEd/Public/Kismet2/KismetReinstanceUtilities.h
--- a/Engine/Source/Editor/UnrealEd/Public/Kismet2/KismetReinstanceUtilities.h
+++ b/Engine/Source/Editor/UnrealEd/Public/Kismet2/KismetReinstanceUtilities.h
@@ -96,6 +96,7 @@ protected:
 		, OriginalCDO(NULL)
 		, bHasReinstanced(false)
 		, bSkipGarbageCollection(false)
+		, ClassToReinstanceDefaultValuesCRC(0)
 	{}
 
 	/** Sets the reinstancer up to work on every object of the specified class */
