UE-8677 Adding a protected component to a class and performing a hot reload can crash the Editor if a Blueprint already exists that is derived from that class

Overriden Archetype for old CDO while hot reloading.

diff --git a/Engine/Source/Developer/HotReload/Private/HotReloadClassReinstancer.cpp b/Engine/Source/Developer/HotReload/Private/HotReloadClassReinstancer.cpp
--- a/Engine/Source/Developer/HotReload/Private/HotReloadClassReinstancer.cpp
+++ b/Engine/Source/Developer/HotReload/Private/HotReloadClassReinstancer.cpp
@@ -229,6 +229,12 @@ void FHotReloadClassReinstancer::RecreateCDOAndSetupOldClassReinstancing(UClass*
 	// Collect the original property values
 	SerializeCDOProperties(InOldClass->GetDefaultObject(), OriginalCDOProperties);
 	
+	FObjectDuplicationParameters Parameters(OriginalCDO, GetTransientPackage());
+	Parameters.DestClass = InOldClass;
+	Parameters.ApplyFlags |= RF_Transient | RF_ArchetypeObject;
+	CopyOfPreviousCDO = StaticDuplicateObjectEx(Parameters);
+
+
 	// Destroy and re-create the CDO, re-running its constructor
 	ReconstructClassDefaultObject(InOldClass);
 
@@ -252,6 +258,12 @@ void FHotReloadClassReinstancer::RecreateCDOAndSetupOldClassReinstancing(UClass*
 				if (!ChildBP->HasAnyFlags(RF_NeedLoad))
 				{
 					Children.AddUnique(ChildBP);
+					auto BPGC = Cast<UBlueprintGeneratedClass>(ChildBP->GeneratedClass);
+					auto CurrentCDO = BPGC ? BPGC->GetDefaultObject(false) : nullptr;
+					if (CurrentCDO && (OriginalCDO == CurrentCDO->GetArchetype()))
+					{
+						BPGC->OverridenArchetypeForCDO = CopyOfPreviousCDO;
+					}
 				}
 			}
 		}
@@ -261,6 +273,7 @@ void FHotReloadClassReinstancer::RecreateCDOAndSetupOldClassReinstancing(UClass*
 FHotReloadClassReinstancer::FHotReloadClassReinstancer(UClass* InNewClass, UClass* InOldClass)
 	: NewClass(nullptr)
 	, bNeedsReinstancing(false)
+	, CopyOfPreviousCDO(nullptr)
 {
 	// If InNewClass is NULL, then the old class has not changed after hot-reload.
 	// However, we still need to check for changes to its constructor code (CDO values).
@@ -432,4 +445,10 @@ void FHotReloadClassReinstancer::ReinstanceObjectsAndUpdateDefaults()
 	UpdateDefaultProperties();
 }
 
+void FHotReloadClassReinstancer::AddReferencedObjects(FReferenceCollector& Collector)
+{
+	FBlueprintCompileReinstancer::AddReferencedObjects(Collector);
+	Collector.AddReferencedObject(CopyOfPreviousCDO);
+}
+
 #endif
\ No newline at end of file
diff --git a/Engine/Source/Developer/HotReload/Private/HotReloadClassReinstancer.h b/Engine/Source/Developer/HotReload/Private/HotReloadClassReinstancer.h
--- a/Engine/Source/Developer/HotReload/Private/HotReloadClassReinstancer.h
+++ b/Engine/Source/Developer/HotReload/Private/HotReloadClassReinstancer.h
@@ -45,6 +45,9 @@ class FHotReloadClassReinstancer : public FBlueprintCompileReinstancer
 	/** True if the provided native class needs re-instancing */
 	bool bNeedsReinstancing;	
 
+	/** Necessary for delta serialization */
+	UObject* CopyOfPreviousCDO;
+
 	/** 
 	 * Sets the re-instancer up for new class re-instancing 
 	 *
@@ -107,6 +110,10 @@ public:
 	{
 		return MakeShareable(new FHotReloadClassReinstancer(InNewClass, InOldClass));
 	}
+
+	// FSerializableObject interface
+	virtual void AddReferencedObjects(FReferenceCollector& Collector) override;
+	// End of FSerializableObject interface
 };
 
 #endif // WITH_ENGINE
\ No newline at end of file
diff --git a/Engine/Source/Runtime/CoreUObject/Private/UObject/Class.cpp b/Engine/Source/Runtime/CoreUObject/Private/UObject/Class.cpp
--- a/Engine/Source/Runtime/CoreUObject/Private/UObject/Class.cpp
+++ b/Engine/Source/Runtime/CoreUObject/Private/UObject/Class.cpp
@@ -3188,6 +3188,12 @@ FArchive& operator<<(FArchive& Ar, FImplementedInterface& A)
 	return Ar;
 }
 
+UObject* UClass::GetArchetypeForCDO() const
+{
+	auto SuperClass = GetSuperClass();
+	return SuperClass ? SuperClass->GetDefaultObject() : nullptr;
+}
+
 void UClass::PurgeClass(bool bRecompilingOnLoad)
 {
 	ClassConstructor = nullptr;
diff --git a/Engine/Source/Runtime/CoreUObject/Private/UObject/UObjectArchetype.cpp b/Engine/Source/Runtime/CoreUObject/Private/UObject/UObjectArchetype.cpp
--- a/Engine/Source/Runtime/CoreUObject/Private/UObject/UObjectArchetype.cpp
+++ b/Engine/Source/Runtime/CoreUObject/Private/UObject/UObjectArchetype.cpp
@@ -13,12 +13,7 @@ UObject* UObject::GetArchetypeFromRequiredInfo(UClass* Class, UObject* Outer, FN
 	const bool bIsCDO = !!(ObjectFlags&RF_ClassDefaultObject);
 	if (bIsCDO)
 	{
-		static UClass* UObjectStaticClass(UObject::StaticClass());
-		if (Class != UObjectStaticClass)
-		{
-			Result = Class->GetSuperClass()->GetDefaultObject(); // the archetype of any CDO (other than the UObject CDO) is the CDO of the super class
-		}
-		// else the archetype of the UObject CDO is NULL
+		Result = Class->GetArchetypeForCDO();
 	}
 	else
 	{
diff --git a/Engine/Source/Runtime/CoreUObject/Public/UObject/Class.h b/Engine/Source/Runtime/CoreUObject/Public/UObject/Class.h
--- a/Engine/Source/Runtime/CoreUObject/Public/UObject/Class.h
+++ b/Engine/Source/Runtime/CoreUObject/Public/UObject/Class.h
@@ -2231,6 +2231,7 @@ public:
 	 */
 	virtual void GetRequiredPreloadDependencies(TArray<UObject*>& DependenciesOut) {}
 
+	virtual UObject* GetArchetypeForCDO() const;
 private:
 	// This signature intentionally hides the method declared in UObjectBaseUtility to make it private.
 	// Call IsChildOf instead; Hidden because calling IsA on a class almost always indicates an error where the caller should use IsChildOf.
diff --git a/Engine/Source/Runtime/Engine/Classes/Engine/BlueprintGeneratedClass.h b/Engine/Source/Runtime/Engine/Classes/Engine/BlueprintGeneratedClass.h
--- a/Engine/Source/Runtime/Engine/Classes/Engine/BlueprintGeneratedClass.h
+++ b/Engine/Source/Runtime/Engine/Classes/Engine/BlueprintGeneratedClass.h
@@ -468,6 +468,11 @@ public:
 	UPROPERTY()
 	TArray<FEventGraphFastCallPair> FastCallPairs;
 
+#if WITH_EDITOR
+	UPROPERTY(Transient)
+	UObject* OverridenArchetypeForCDO;
+#endif //WITH_EDITOR
+
 	/** 
 	 * Gets an array of all BPGeneratedClasses (including InClass as 0th element) parents of given generated class 
 	 *
@@ -495,6 +500,7 @@ public:
 #if WITH_EDITOR
 	virtual UClass* GetAuthoritativeClass() override;
 	virtual void ConditionalRecompileClass(TArray<UObject*>* ObjLoaded) override;
+	virtual UObject* GetArchetypeForCDO() const override;
 #endif //WITH_EDITOR
 	virtual bool IsFunctionImplementedInBlueprint(FName InFunctionName) const override;
 	virtual uint8* GetPersistentUberGraphFrame(UObject* Obj, UFunction* FuncToCheck) const override;
diff --git a/Engine/Source/Runtime/Engine/Private/BlueprintGeneratedClass.cpp b/Engine/Source/Runtime/Engine/Private/BlueprintGeneratedClass.cpp
--- a/Engine/Source/Runtime/Engine/Private/BlueprintGeneratedClass.cpp
+++ b/Engine/Source/Runtime/Engine/Private/BlueprintGeneratedClass.cpp
@@ -265,6 +265,17 @@ void UBlueprintGeneratedClass::ConditionalRecompileClass(TArray<UObject*>* ObjLo
 		}
 	}
 }
+
+UObject* UBlueprintGeneratedClass::GetArchetypeForCDO() const
+{
+	if (OverridenArchetypeForCDO)
+	{
+		ensure(OverridenArchetypeForCDO->IsA(GetSuperClass()));
+		return OverridenArchetypeForCDO;
+	}
+
+	return Super::GetArchetypeForCDO();
+}
 #endif //WITH_EDITOR
 
 bool UBlueprintGeneratedClass::IsFunctionImplementedInBlueprint(FName InFunctionName) const
@@ -653,6 +664,7 @@ void UBlueprintGeneratedClass::PurgeClass(bool bRecompilingOnLoad)
 
 	UberGraphFramePointerProperty = NULL;
 	UberGraphFunction = NULL;
+	OverridenArchetypeForCDO = NULL;
 
 #if UE_BLUEPRINT_EVENTGRAPH_FASTCALLS
 	FastCallPairs.Empty();
