Fixing ensures when performing hot-reload

diff --git a/Engine/Source/Developer/HotReload/Private/HotReloadClassReinstancer.cpp b/Engine/Source/Developer/HotReload/Private/HotReloadClassReinstancer.cpp
--- a/Engine/Source/Developer/HotReload/Private/HotReloadClassReinstancer.cpp
+++ b/Engine/Source/Developer/HotReload/Private/HotReloadClassReinstancer.cpp
@@ -232,8 +232,10 @@ void FHotReloadClassReinstancer::RecreateCDOAndSetupOldClassReinstancing(UClass*
 	FObjectDuplicationParameters Parameters(OriginalCDO, GetTransientPackage());
 	Parameters.DestClass = InOldClass;
 	Parameters.ApplyFlags |= RF_Transient | RF_ArchetypeObject;
-	CopyOfPreviousCDO = StaticDuplicateObjectEx(Parameters);
-
+	{
+		TGuardValue<bool> GuardGIsDuplicatingClassForReinstancing(GIsDuplicatingClassForReinstancing, true);
+		CopyOfPreviousCDO = StaticDuplicateObjectEx(Parameters);
+	}
 
 	// Destroy and re-create the CDO, re-running its constructor
 	ReconstructClassDefaultObject(InOldClass);
