Fixing hot-reload crashing when FStringAssetReference changes during hot-reload

diff --git a/Engine/Source/Developer/HotReload/Private/HotReloadClassReinstancer.cpp b/Engine/Source/Developer/HotReload/Private/HotReloadClassReinstancer.cpp
--- a/Engine/Source/Developer/HotReload/Private/HotReloadClassReinstancer.cpp
+++ b/Engine/Source/Developer/HotReload/Private/HotReloadClassReinstancer.cpp
@@ -344,6 +344,37 @@ void FHotReloadClassReinstancer::UpdateDefaultProperties()
 			}
 			return *this;
 		}
+		virtual FArchive& operator<<(FName& InName) override
+		{
+			FArchive& Ar = *this;
+			NAME_INDEX ComparisonIndex = InName.GetComparisonIndex();
+			NAME_INDEX DisplayIndex = InName.GetDisplayIndex();
+			int32 Number = InName.GetNumber();
+			Ar << ComparisonIndex;
+			Ar << DisplayIndex;
+			Ar << Number;
+			return Ar;
+		}
+		virtual FArchive& operator<<(FLazyObjectPtr& LazyObjectPtr) override
+		{
+			FArchive& Ar = *this;
+			auto UniqueID = LazyObjectPtr.GetUniqueID();
+			Ar << UniqueID;
+			return *this;
+		}
+		virtual FArchive& operator<<(FAssetPtr& AssetPtr) override
+		{
+			FArchive& Ar = *this;
+			auto UniqueID = AssetPtr.GetUniqueID();
+			Ar << UniqueID;
+			return Ar;
+		}
+		virtual FArchive& operator<<(FStringAssetReference& Value) override
+		{
+			FArchive& Ar = *this;
+			Ar << Value.AssetLongPathname;
+			return Ar;
+		}
 	};
 
 	// Collect default subobjects to update their properties too
