UE-10163 If a custom code class is modified and a Hot Reload performed, any existing unsaved Blueprints derived from that class can no longer be saved.

References to hot-reloaded class are updated in blueprints (before recompiling and reinstancing).
I guess it could be optimized if necessary.

diff --git a/Engine/Source/Developer/HotReload/Private/HotReloadClassReinstancer.cpp b/Engine/Source/Developer/HotReload/Private/HotReloadClassReinstancer.cpp
--- a/Engine/Source/Developer/HotReload/Private/HotReloadClassReinstancer.cpp
+++ b/Engine/Source/Developer/HotReload/Private/HotReloadClassReinstancer.cpp
@@ -282,6 +282,14 @@ FHotReloadClassReinstancer::FHotReloadClassReinstancer(UClass* InNewClass, UClas
 	if (InNewClass)
 	{
 		SetupNewClassReinstancing(InNewClass, InOldClass);
+
+		TMap<UObject*, UObject*> ClassRedirects;
+		ClassRedirects.Add(InOldClass, InNewClass);
+
+		for (TObjectIterator<UBlueprint> BlueprintIt; BlueprintIt; ++BlueprintIt)
+		{
+			FArchiveReplaceObjectRef<UObject>(*BlueprintIt, ClassRedirects, false, true, true);
+		}
 	}
 	else
 	{
