UE-13981 Input pins of BP node are not updated after a hot reload

diff --git a/Engine/Source/Developer/HotReload/Private/HotReloadClassReinstancer.cpp b/Engine/Source/Developer/HotReload/Private/HotReloadClassReinstancer.cpp
--- a/Engine/Source/Developer/HotReload/Private/HotReloadClassReinstancer.cpp
+++ b/Engine/Source/Developer/HotReload/Private/HotReloadClassReinstancer.cpp
@@ -278,6 +278,11 @@ FHotReloadClassReinstancer::FHotReloadClassReinstancer(UClass* InNewClass, UClas
 	, bNeedsReinstancing(false)
 	, CopyOfPreviousCDO(nullptr)
 {
+	ensure(InOldClass);
+	ensure(!HotReloadedOldClass && !HotReloadedNewClass);
+	HotReloadedOldClass = InOldClass;
+	HotReloadedNewClass = InNewClass ? InNewClass : InOldClass;
+
 	// If InNewClass is NULL, then the old class has not changed after hot-reload.
 	// However, we still need to check for changes to its constructor code (CDO values).
 	if (InNewClass)
@@ -289,7 +294,11 @@ FHotReloadClassReinstancer::FHotReloadClassReinstancer(UClass* InNewClass, UClas
 
 		for (TObjectIterator<UBlueprint> BlueprintIt; BlueprintIt; ++BlueprintIt)
 		{
-			FArchiveReplaceObjectRef<UObject>(*BlueprintIt, ClassRedirects, false, true, true);
+			FArchiveReplaceObjectRef<UObject> ReplaceObjectArch(*BlueprintIt, ClassRedirects, false, true, true);
+			if (ReplaceObjectArch.GetCount())
+			{
+				EnlistDependentBlueprintToRecompile(*BlueprintIt, true);
+			}
 		}
 	}
 	else
@@ -303,6 +312,10 @@ FHotReloadClassReinstancer::~FHotReloadClassReinstancer()
 	// Make sure the base class does not remove the DuplicatedClass from root, we not always want it.
 	// For example when we're just reconstructing CDOs. Other cases are handled by HotReloadClassReinstancer.
 	DuplicatedClass = nullptr;
+
+	ensure(HotReloadedOldClass);
+	HotReloadedOldClass = nullptr;
+	HotReloadedNewClass = nullptr;
 }
 
 /** Helper for finding subobject in an array. Usually there's not that many subobjects on a class to justify a TMap */
diff --git a/Engine/Source/Editor/BlueprintGraph/Private/K2Node_CallFunction.cpp b/Engine/Source/Editor/BlueprintGraph/Private/K2Node_CallFunction.cpp
--- a/Engine/Source/Editor/BlueprintGraph/Private/K2Node_CallFunction.cpp
+++ b/Engine/Source/Editor/BlueprintGraph/Private/K2Node_CallFunction.cpp
@@ -2170,14 +2170,15 @@ FText UK2Node_CallFunction::GetMenuCategory() const
 
 bool UK2Node_CallFunction::HasExternalBlueprintDependencies(TArray<class UStruct*>* OptionalOutput) const
 {
-	const UClass* SourceClass = FunctionReference.GetMemberParentClass(GetBlueprintClassFromNode());
+	UFunction* Function = GetTargetFunction();
+	const UClass* SourceClass = Function ? Function->GetOwnerClass() : nullptr;
 	const UBlueprint* SourceBlueprint = GetBlueprint();
-	const bool bResult = (SourceClass != NULL) && (SourceClass->ClassGeneratedBy != NULL) && (SourceClass->ClassGeneratedBy != SourceBlueprint);
+	const bool bResult = (SourceClass != NULL) && (SourceClass->ClassGeneratedBy != SourceBlueprint);
 	if (bResult && OptionalOutput)
 	{
-		OptionalOutput->Add(GetTargetFunction());
+		OptionalOutput->AddUnique(Function);
 	}
-	return bResult || Super::HasExternalBlueprintDependencies(OptionalOutput);
+	return Super::HasExternalBlueprintDependencies(OptionalOutput) || bResult;
 }
 
 UEdGraph* UK2Node_CallFunction::GetFunctionGraph(const UEdGraphNode*& OutGraphNode) const
diff --git a/Engine/Source/Editor/KismetCompiler/Private/KismetCompilerModule.cpp b/Engine/Source/Editor/KismetCompiler/Private/KismetCompilerModule.cpp
--- a/Engine/Source/Editor/KismetCompiler/Private/KismetCompilerModule.cpp
+++ b/Engine/Source/Editor/KismetCompiler/Private/KismetCompilerModule.cpp
@@ -160,6 +160,8 @@ void FKismet2CompilerModule::CompileBlueprint(class UBlueprint* Blueprint, const
 	{
 		BP_SCOPED_COMPILER_EVENT_STAT(EKismetCompilerStats_CompileGeneratedClass);
 
+		FBlueprintCompileReinstancer::OptionallyRefreshNodes(Blueprint);
+
 		// Perform the full compile
 		CompileBlueprintInner(Blueprint, CompileOptions, Results, ParentReinstancer, ObjLoaded);
 
diff --git a/Engine/Source/Editor/UnrealEd/Private/Kismet2/KismetReinstanceUtilities.cpp b/Engine/Source/Editor/UnrealEd/Private/Kismet2/KismetReinstanceUtilities.cpp
--- a/Engine/Source/Editor/UnrealEd/Private/Kismet2/KismetReinstanceUtilities.cpp
+++ b/Engine/Source/Editor/UnrealEd/Private/Kismet2/KismetReinstanceUtilities.cpp
@@ -130,6 +130,9 @@ TSet<TWeakObjectPtr<UBlueprint>> FBlueprintCompileReinstancer::DependentBlueprin
 TSet<TWeakObjectPtr<UBlueprint>> FBlueprintCompileReinstancer::DependentBlueprintsToRecompile = TSet<TWeakObjectPtr<UBlueprint>>();
 TSet<TWeakObjectPtr<UBlueprint>> FBlueprintCompileReinstancer::DependentBlueprintsToByteRecompile = TSet<TWeakObjectPtr<UBlueprint>>();
 
+UClass* FBlueprintCompileReinstancer::HotReloadedOldClass = nullptr;
+UClass* FBlueprintCompileReinstancer::HotReloadedNewClass = nullptr;
+
 FBlueprintCompileReinstancer::FBlueprintCompileReinstancer(UClass* InClassToReinstance, bool bIsBytecodeOnly, bool bSkipGC)
 	: ClassToReinstance(InClassToReinstance)
 	, DuplicatedClass(NULL)
@@ -293,6 +296,22 @@ void FBlueprintCompileReinstancer::AddReferencedObjects(FReferenceCollector& Col
 	Collector.AddReferencedObject(DuplicatedClass);
 }
 
+void FBlueprintCompileReinstancer::OptionallyRefreshNodes(UBlueprint* CurrentBP)
+{
+	if (HotReloadedNewClass)
+	{
+		UPackage* const Package = Cast<UPackage>(CurrentBP->GetOutermost());
+		const bool bStartedWithUnsavedChanges = Package != nullptr ? Package->IsDirty() : true;
+
+		FBlueprintEditorUtils::RefreshExternalBlueprintDependencyNodes(CurrentBP, HotReloadedNewClass);
+
+		if (Package != nullptr && Package->IsDirty() && !bStartedWithUnsavedChanges)
+		{
+			Package->SetDirtyFlag(false);
+		}
+	}
+}
+
 FBlueprintCompileReinstancer::~FBlueprintCompileReinstancer()
 {
 }
diff --git a/Engine/Source/Editor/UnrealEd/Public/Kismet2/KismetReinstanceUtilities.h b/Engine/Source/Editor/UnrealEd/Public/Kismet2/KismetReinstanceUtilities.h
--- a/Engine/Source/Editor/UnrealEd/Public/Kismet2/KismetReinstanceUtilities.h
+++ b/Engine/Source/Editor/UnrealEd/Public/Kismet2/KismetReinstanceUtilities.h
@@ -20,6 +20,9 @@ protected:
 	static TSet<TWeakObjectPtr<UBlueprint>> DependentBlueprintsToRecompile;
 	static TSet<TWeakObjectPtr<UBlueprint>> DependentBlueprintsToByteRecompile;
 
+	static UClass* HotReloadedOldClass;
+	static UClass* HotReloadedNewClass;
+
 	/** Reference to the class we're actively reinstancing */
 	UClass* ClassToReinstance;
 
@@ -58,6 +61,8 @@ public:
 	virtual void AddReferencedObjects(FReferenceCollector& Collector) override;
 	// End of FSerializableObject interface
 
+	static void OptionallyRefreshNodes(UBlueprint* BP);
+
 	void ListDependentBlueprintsToRefresh(const TArray<UBlueprint*>& DependentBPs);
 	void EnlistDependentBlueprintToRecompile(UBlueprint* BP, bool bBytecodeOnly);
 	void BlueprintWasRecompiled(UBlueprint* BP, bool bBytecodeOnly);
