UE-15366: Hot-reload crashes if you try to change and reload a default subobject's class.

diff --git a/Engine/Source/Developer/HotReload/Private/HotReload.cpp b/Engine/Source/Developer/HotReload/Private/HotReload.cpp
--- a/Engine/Source/Developer/HotReload/Private/HotReload.cpp
+++ b/Engine/Source/Developer/HotReload/Private/HotReload.cpp
@@ -620,10 +620,34 @@ void FHotReloadModule::DoHotReloadCallback(bool bRecompileFinished, ECompilation
 	DoHotReloadInternal(bRecompileFinished, CompilationResult, Packages, InDependentModules, HotReloadAr);
 }
 
+#if WITH_HOT_RELOAD
+/**
+ * Gets duplicated CDO from the cache, renames it and returns.
+ */
+UObject* GetCachedCDODuplicate(UClass* Class, FName Name)
+{
+	UObject* DupCDO = nullptr;
+
+	UObject** DupCDOPtr = GetDuplicatedCDOMap().Find(Class);
+	if (DupCDOPtr != nullptr)
+	{
+		DupCDO = *DupCDOPtr;
+		DupCDO->Rename(*Name.ToString(), GetTransientPackage(),
+			REN_DoNotDirty | REN_DontCreateRedirectors | REN_ForceNoResetLoaders | REN_NonTransactional | REN_SkipGeneratedClasses);
+	}
+
+	return DupCDO;
+}
+#endif // WITH_HOT_RELOAD
+
 ECompilationResult::Type FHotReloadModule::DoHotReloadInternal(bool bRecompileFinished, ECompilationResult::Type CompilationResult, TArray<UPackage*> Packages, TArray< FName > InDependentModules, FOutputDevice &HotReloadAr)
 {
 	ECompilationResult::Type Result = ECompilationResult::Unsupported;
 #if WITH_HOT_RELOAD
+	FBlueprintCompileReinstancer::FCDODuplicatesProvider& CDODuplicatesProvider = FBlueprintCompileReinstancer::GetCDODuplicatesProviderDelegate();
+
+	CDODuplicatesProvider.BindStatic(&GetCachedCDODuplicate);
+
 	if (CompilationResult == ECompilationResult::Succeeded)
 	{
 		FFeedbackContext& ErrorsFC = UClass::GetDefaultPropertiesFeedbackContext();
@@ -741,6 +765,9 @@ ECompilationResult::Type FHotReloadModule::DoHotReloadInternal(bool bRecompileFi
 		HotReloadAr.Logf(ELogVerbosity::Warning, TEXT("HotReload failed, recompile failed"));
 		Result = ECompilationResult::OtherCompilationError;
 	}
+
+	CDODuplicatesProvider.Unbind();
+	GetDuplicatedCDOMap().Empty();
 #endif
 	bIsHotReloadingFromEditor = false;
 	return Result;
diff --git a/Engine/Source/Developer/HotReload/Private/HotReloadClassReinstancer.cpp b/Engine/Source/Developer/HotReload/Private/HotReloadClassReinstancer.cpp
--- a/Engine/Source/Developer/HotReload/Private/HotReloadClassReinstancer.cpp
+++ b/Engine/Source/Developer/HotReload/Private/HotReloadClassReinstancer.cpp
@@ -177,43 +177,22 @@ void FHotReloadClassReinstancer::SerializeCDOProperties(UObject* InObject, FHotR
 	FCDOWriter Ar(OutData, InObject, VisitedObjects);
 }
 
-void FHotReloadClassReinstancer::ReconstructClassDefaultObject(UClass* InOldClass)
+void FHotReloadClassReinstancer::ReconstructClassDefaultObject(UClass* InClass, UObject* InOuter, FName InName, EObjectFlags InFlags)
 {
-	// Remember all the basic info about the object before we destroy it
-	UObject* OldCDO = InOldClass->GetDefaultObject();
-	EObjectFlags CDOFlags = OldCDO->GetFlags();
-	UObject* CDOOuter = OldCDO->GetOuter();
-	FName CDOName = OldCDO->GetFName();
-
 	// Get the parent CDO
-	UClass* ParentClass = InOldClass->GetSuperClass();
+	UClass* ParentClass = InClass->GetSuperClass();
 	UObject* ParentDefaultObject = NULL;
 	if (ParentClass != NULL)
 	{
 		ParentDefaultObject = ParentClass->GetDefaultObject(); // Force the default object to be constructed if it isn't already
 	}
 
-	if (!OldCDO->HasAnyFlags(RF_FinishDestroyed))
-	{
-		// Begin the asynchronous object cleanup.
-		OldCDO->ConditionalBeginDestroy();
-
-		// Wait for the object's asynchronous cleanup to finish.
-		while (!OldCDO->IsReadyForFinishDestroy())
-		{
-			FPlatformProcess::Sleep(0);
-		}
-		// Finish destroying the object.
-		OldCDO->ConditionalFinishDestroy();
-	}
-	OldCDO->~UObject();
-
 	// Re-create
-	FMemory::Memzero((void*)OldCDO, InOldClass->GetPropertiesSize());
-	new ((void *)OldCDO) UObjectBase(InOldClass, CDOFlags, CDOOuter, CDOName);
+	InClass->ClassDefaultObject = StaticAllocateObject(InClass, InOuter, InName, InFlags, false);
+	check(InClass->ClassDefaultObject);
 	const bool bShouldInitilizeProperties = false;
 	const bool bCopyTransientsFromClassDefaults = false;
-	(*InOldClass->ClassConstructor)(FObjectInitializer(OldCDO, ParentDefaultObject, bCopyTransientsFromClassDefaults, bShouldInitilizeProperties));
+	(*InClass->ClassConstructor)(FObjectInitializer(InClass->ClassDefaultObject, ParentDefaultObject, bCopyTransientsFromClassDefaults, bShouldInitilizeProperties));
 }
 
 void FHotReloadClassReinstancer::RecreateCDOAndSetupOldClassReinstancing(UClass* InOldClass)
@@ -229,17 +208,25 @@ void FHotReloadClassReinstancer::RecreateCDOAndSetupOldClassReinstancing(UClass*
 
 	// Collect the original property values
 	SerializeCDOProperties(InOldClass->GetDefaultObject(), OriginalCDOProperties);
-	
-	FObjectDuplicationParameters Parameters(OriginalCDO, GetTransientPackage());
-	Parameters.DestClass = InOldClass;
-	Parameters.ApplyFlags |= RF_Transient | RF_ArchetypeObject;
-	{
-		TGuardValue<bool> GuardGIsDuplicatingClassForReinstancing(GIsDuplicatingClassForReinstancing, true);
-		CopyOfPreviousCDO = StaticDuplicateObjectEx(Parameters);
-	}
 
-	// Destroy and re-create the CDO, re-running its constructor
-	ReconstructClassDefaultObject(InOldClass);
+	// Remember all the basic info about the object before we rename it
+	EObjectFlags CDOFlags = OriginalCDO->GetFlags();
+	UObject* CDOOuter = OriginalCDO->GetOuter();
+	FName CDOName = OriginalCDO->GetFName();
+
+	// Rename original CDO, so we can store this one as OverridenArchetypeForCDO
+	// and create new one with the same name and outer.
+	OriginalCDO->Rename(
+		*MakeUniqueObjectName(
+			GetTransientPackage(),
+			OriginalCDO->GetClass(),
+			*FString::Printf(TEXT("BPGC_ARCH_FOR_CDO_%s"), *InOldClass->GetName())
+		).ToString(),
+		GetTransientPackage(),
+		REN_DoNotDirty | REN_DontCreateRedirectors | REN_NonTransactional | REN_SkipGeneratedClasses | REN_ForceNoResetLoaders);
+	
+	// Re-create the CDO, re-running its constructor
+	ReconstructClassDefaultObject(InOldClass, CDOOuter, CDOName, CDOFlags);
 
 	// Collect the property values after re-constructing the CDO
 	SerializeCDOProperties(InOldClass->GetDefaultObject(), ReconstructedCDOProperties);
@@ -265,7 +252,7 @@ void FHotReloadClassReinstancer::RecreateCDOAndSetupOldClassReinstancing(UClass*
 					auto CurrentCDO = BPGC ? BPGC->GetDefaultObject(false) : nullptr;
 					if (CurrentCDO && (OriginalCDO == CurrentCDO->GetArchetype()))
 					{
-						BPGC->OverridenArchetypeForCDO = CopyOfPreviousCDO;
+						BPGC->OverridenArchetypeForCDO = OriginalCDO;
 					}
 				}
 			}
diff --git a/Engine/Source/Developer/HotReload/Private/HotReloadClassReinstancer.h b/Engine/Source/Developer/HotReload/Private/HotReloadClassReinstancer.h
--- a/Engine/Source/Developer/HotReload/Private/HotReloadClassReinstancer.h
+++ b/Engine/Source/Developer/HotReload/Private/HotReloadClassReinstancer.h
@@ -72,11 +72,14 @@ class FHotReloadClassReinstancer : public FBlueprintCompileReinstancer
 	void SerializeCDOProperties(UObject* InObject, FCDOPropertyData& OutData);
 
 	/**
-	* Re-creates class default object
-	*
-	* @param InOldClass Class that has NOT changed after hot-reload
-	*/
-	void ReconstructClassDefaultObject(UClass* InOldClass);
+	 * Re-creates class default object.
+	 *
+	 * @param InClass Class that has NOT changed after hot-reload.
+	 * @param InOuter Outer for the new CDO.
+	 * @param InName Name of the new CDO.
+	 * @param InFlags Flags of the new CDO.
+	 */
+	void ReconstructClassDefaultObject(UClass* InClass, UObject* InOuter, FName InName, EObjectFlags InFlags);
 
 	/** Updates property values on instances of the hot-reloaded class */
 	void UpdateDefaultProperties();
diff --git a/Engine/Source/Editor/UnrealEd/Private/Kismet2/KismetReinstanceUtilities.cpp b/Engine/Source/Editor/UnrealEd/Private/Kismet2/KismetReinstanceUtilities.cpp
--- a/Engine/Source/Editor/UnrealEd/Private/Kismet2/KismetReinstanceUtilities.cpp
+++ b/Engine/Source/Editor/UnrealEd/Private/Kismet2/KismetReinstanceUtilities.cpp
@@ -178,16 +178,13 @@ FBlueprintCompileReinstancer::FBlueprintCompileReinstancer(UClass* InClassToRein
 		DuplicatedClass->StaticLink(true);
 
 		// Copy over the ComponentNametoDefaultObjectMap, which tells CopyPropertiesForUnrelatedObjects which components are instanced and which aren't
-
-		GIsDuplicatingClassForReinstancing = true;
-		UObject* OldCDO = ClassToReinstance->GetDefaultObject();
-		DuplicatedClass->ClassDefaultObject = (UObject*)StaticDuplicateObject(OldCDO, GetTransientPackage(), *DuplicatedClass->GetDefaultObjectName().ToString());
-		GIsDuplicatingClassForReinstancing = false;
+		
+		DuplicatedClass->ClassDefaultObject = GetClassCDODuplicate(ClassToReinstance, DuplicatedClass->GetDefaultObjectName());
 
 		DuplicatedClass->ClassDefaultObject->SetFlags(RF_ClassDefaultObject);
 		DuplicatedClass->ClassDefaultObject->SetClass(DuplicatedClass);
 
-		OldCDO->SetClass(DuplicatedClass);
+		ClassToReinstance->GetDefaultObject()->SetClass(DuplicatedClass);
 		ObjectsThatShouldUseOldStuff.Add(DuplicatedClass); //CDO of REINST_ class can be used as archetype
 
 		if( !bIsBytecodeOnly )
@@ -1283,3 +1280,25 @@ void FBlueprintCompileReinstancer::ReparentChild(UClass* ChildClass)
 	ChildClass->Bind();
 	ChildClass->StaticLink(true);
 }
+
+UObject* FBlueprintCompileReinstancer::GetClassCDODuplicate(UClass* Class, FName Name)
+{
+	UObject* DupCDO = nullptr;
+
+	FCDODuplicatesProvider& CDODupProvider = GetCDODuplicatesProviderDelegate();
+
+	if (!CDODupProvider.IsBound() || (DupCDO = CDODupProvider.Execute(Class, Name)) == nullptr)
+	{
+		GIsDuplicatingClassForReinstancing = true;
+		DupCDO = (UObject*)StaticDuplicateObject(Class->GetDefaultObject(), GetTransientPackage(), *Name.ToString());
+		GIsDuplicatingClassForReinstancing = false;
+	}
+
+	return DupCDO;
+}
+
+FBlueprintCompileReinstancer::FCDODuplicatesProvider& FBlueprintCompileReinstancer::GetCDODuplicatesProviderDelegate()
+{
+	static FCDODuplicatesProvider Delegate;
+	return Delegate;
+}
diff --git a/Engine/Source/Editor/UnrealEd/Public/Kismet2/KismetReinstanceUtilities.h b/Engine/Source/Editor/UnrealEd/Public/Kismet2/KismetReinstanceUtilities.h
--- a/Engine/Source/Editor/UnrealEd/Public/Kismet2/KismetReinstanceUtilities.h
+++ b/Engine/Source/Editor/UnrealEd/Public/Kismet2/KismetReinstanceUtilities.h
@@ -14,6 +14,17 @@ class FReinstanceFinalizer;
 
 class UNREALED_API FBlueprintCompileReinstancer : public TSharedFromThis<FBlueprintCompileReinstancer>, public FGCObject
 {
+public:
+	/**
+	 * CDO duplicates provider delegate type.
+	 */
+	DECLARE_DELEGATE_RetVal_TwoParams(UObject*, FCDODuplicatesProvider, UClass*, FName);
+
+	/**
+	 * Gets CDO duplicates provider delegate.
+	 */
+	static FCDODuplicatesProvider& GetCDODuplicatesProviderDelegate();
+
 protected:
 
 	static TSet<TWeakObjectPtr<UBlueprint>> DependentBlueprintsToRefresh;
@@ -131,6 +142,16 @@ protected:
 
 	/** Determine whether reinstancing actors should preserve the root component of the new actor */
 	virtual bool ShouldPreserveRootComponentOfReinstancedActor() const { return true; }
+
+private:
+	/**
+	 * Gets class CDO duplicate from cache (or creates it if not available or not
+	 * during hot-reload) for given class.
+	 *
+	 * @param Class Class to get (or create) CDO for.
+	 * @param Name The name that CDO has to be renamed with (or created with).
+	 */
+	static UObject* GetClassCDODuplicate(UClass* Class, FName Name);
 };
 
 
diff --git a/Engine/Source/Runtime/CoreUObject/Private/UObject/UObjectBase.cpp b/Engine/Source/Runtime/CoreUObject/Private/UObject/UObjectBase.cpp
--- a/Engine/Source/Runtime/CoreUObject/Private/UObject/UObjectBase.cpp
+++ b/Engine/Source/Runtime/CoreUObject/Private/UObject/UObjectBase.cpp
@@ -523,6 +523,12 @@ static TArray<FFieldCompiledInInfo*>& GetDeferredClassRegistration()
 }
 
 #if WITH_HOT_RELOAD
+TMap<UClass*, UObject*>& GetDuplicatedCDOMap()
+{
+	static TMap<UClass*, UObject*> Map;
+	return Map;
+}
+
 /** Map of deferred class registration info (including size and reflection info) */
 static TMap<FName, FFieldCompiledInInfo*>& GetDeferRegisterClassMap()
 {
@@ -645,6 +651,37 @@ void UClassReplaceHotReloadClasses()
 	}
 	GetHotReloadClasses().Empty();
 }
+
+/**
+ * Creates a cache of cpp-only-changed UClasses' CDOs, which are going to be
+ * used later during BP reinstancing.
+ */
+static void UClassGenerateCDODuplicatesForHotReload()
+{
+	if (!GIsHotReload)
+	{
+		return;
+	}
+
+	for (FObjectIterator It(UClass::StaticClass()); It; ++It)
+	{
+		UClass* Class = (UClass*)*It;
+
+		for (auto* HotReloadedClass : GetHotReloadClasses())
+		{
+			if (!HotReloadedClass->bHasChanged && Class->IsChildOf(HotReloadedClass->OldClass))
+			{
+				GIsDuplicatingClassForReinstancing = true;
+				UObject* DupCDO = (UObject*)StaticDuplicateObject(
+					Class->GetDefaultObject(), GetTransientPackage(),
+					*MakeUniqueObjectName(GetTransientPackage(), Class, TEXT("HOTRELOAD_CDO_DUPLICATE")).ToString()
+					);
+				GIsDuplicatingClassForReinstancing = false;
+				GetDuplicatedCDOMap().Add(Class, DupCDO);
+			}
+		}
+	}
+}
 #endif
 
 /**
@@ -747,6 +784,9 @@ void ProcessNewlyLoadedUObjects()
 {
 	DECLARE_SCOPE_CYCLE_COUNTER(TEXT("ProcessNewlyLoadedUObjects"), STAT_ProcessNewlyLoadedUObjects, STATGROUP_ObjectVerbose);
 
+#if WITH_HOT_RELOAD
+	UClassGenerateCDODuplicatesForHotReload();
+#endif
 	UClassRegisterAllCompiledInClasses();
 
 	while( AnyNewlyLoadedUObjects() )
diff --git a/Engine/Source/Runtime/CoreUObject/Public/UObject/UObjectBase.h b/Engine/Source/Runtime/CoreUObject/Public/UObject/UObjectBase.h
--- a/Engine/Source/Runtime/CoreUObject/Public/UObject/UObjectBase.h
+++ b/Engine/Source/Runtime/CoreUObject/Public/UObject/UObjectBase.h
@@ -378,6 +378,11 @@ COREUOBJECT_API bool AnyNewlyLoadedUObjects();
 /** Must be called after a module has been loaded that contains UObject classes */
 COREUOBJECT_API void ProcessNewlyLoadedUObjects();
 
+#if WITH_HOT_RELOAD
+/** Map of duplicated CDOs for reinstancing during hot-reload purposes. */
+COREUOBJECT_API TMap<UClass*, UObject*>& GetDuplicatedCDOMap();
+#endif // WITH_HOT_RELOAD
+
 /**
  * Final phase of UObject initialization. all auto register objects are added to the main data structures.
  */
