UE-18307: Hot reloading with a plugin in the project prevents the project from being able to save.

diff --git a/Engine/Source/Developer/HotReload/Private/HotReload.cpp b/Engine/Source/Developer/HotReload/Private/HotReload.cpp
--- a/Engine/Source/Developer/HotReload/Private/HotReload.cpp
+++ b/Engine/Source/Developer/HotReload/Private/HotReload.cpp
@@ -151,6 +151,13 @@ private:
 	 */
 	ECompilationResult::Type DoHotReloadInternal(bool bRecompileFinished, ECompilationResult::Type CompilationResult, TArray<UPackage*> Packages, TArray<FName> InDependentModules, FOutputDevice &HotReloadAr);
 
+	/**
+	 * Finds all references to old CDOs and replaces them with the new ones.
+	 * Skipping UBlueprintGeneratedClass::OverridenArchetypeForCDO as it's the
+	 * only one needed.
+	 */
+	void ReplaceReferencesToReconstructedCDOs();
+
 	/**
 	* Callback for async ompilation
 	*/
@@ -350,6 +357,9 @@ private:
 
 	/** True if the directory watcher has been successfully initialized */
 	bool bDirectoryWatcherInitialized;
+
+	/** Reconstructed CDOs map during hot-reload. */
+	TMap<UObject*, UObject*> ReconstructedCDOsMap;
 };
 
 namespace HotReloadDefs
@@ -756,6 +766,9 @@ ECompilationResult::Type FHotReloadModule::DoHotReloadInternal(bool bRecompileFi
 			HotReloadAr.Logf(ELogVerbosity::Warning, TEXT("HotReload successful (%d functions remapped  %d scriptstructs remapped)"), Count, ScriptStructs.Num());
 
 			HotReloadFunctionRemap.Empty();
+
+			ReplaceReferencesToReconstructedCDOs();
+
 			Result = ECompilationResult::Succeeded;
 		}
 
@@ -778,6 +791,45 @@ ECompilationResult::Type FHotReloadModule::DoHotReloadInternal(bool bRecompileFi
 	return Result;
 }
 
+void FHotReloadModule::ReplaceReferencesToReconstructedCDOs()
+{
+	if (ReconstructedCDOsMap.Num() == 0)
+	{
+		return;
+	}
+
+	TArray<UObject*> OldCDOs;
+	ReconstructedCDOsMap.GetKeys(OldCDOs);
+
+	for (FObjectIterator ObjIter; ObjIter; ++ObjIter)
+	{
+		UObject* CurObject = *ObjIter;
+
+		FFindReferencersArchive FindRefsArchive(CurObject, OldCDOs);
+
+		TMap<UObject*, int32> ReferenceCounts;
+		TMultiMap<UObject*, UProperty*> ReferencingProperties;
+
+		FindRefsArchive.GetReferenceCounts(ReferenceCounts, ReferencingProperties);
+
+		for (auto& ReferencingProperty : ReferencingProperties)
+		{
+			static const UProperty* PropertyToSkip = UBlueprintGeneratedClass::StaticClass()->FindPropertyByName(TEXT("OverridenArchetypeForCDO"));
+			if (!ReferencingProperty.Value->IsA<UObjectProperty>() || ReferencingProperty.Value == PropertyToSkip)
+			{
+				continue;
+			}
+
+			UObject* OldCDO = ReferencingProperty.Key;
+			UObjectProperty* Prop = (UObjectProperty*)ReferencingProperty.Value;
+
+			Prop->SetObjectPropertyValue((uint8*)CurObject + Prop->GetOffset_ForInternal(), ReconstructedCDOsMap[OldCDO]);
+		}
+	}
+
+	ReconstructedCDOsMap.Empty();
+}
+
 ECompilationResult::Type FHotReloadModule::RebindPackages(TArray<UPackage*> InPackages, TArray<FName> DependentModules, const bool bWaitForCompletion, FOutputDevice &Ar)
 {
 	ECompilationResult::Type Result = ECompilationResult::Unknown;
@@ -877,7 +929,7 @@ ECompilationResult::Type FHotReloadModule::RebindPackagesInternal(TArray<UPackag
 #if WITH_ENGINE
 void FHotReloadModule::ReinstanceClass(UClass* OldClass, UClass* NewClass)
 {	
-	auto ReinstanceHelper = FHotReloadClassReinstancer::Create(NewClass, OldClass);
+	auto ReinstanceHelper = FHotReloadClassReinstancer::Create(NewClass, OldClass, ReconstructedCDOsMap);
 	if (ReinstanceHelper->ClassNeedsReinstancing())
 	{
 		UE_LOG(LogHotReload, Log, TEXT("Re-instancing %s after hot-reload."), NewClass ? *NewClass->GetName() : *OldClass->GetName());
diff --git a/Engine/Source/Developer/HotReload/Private/HotReloadClassReinstancer.cpp b/Engine/Source/Developer/HotReload/Private/HotReloadClassReinstancer.cpp
--- a/Engine/Source/Developer/HotReload/Private/HotReloadClassReinstancer.cpp
+++ b/Engine/Source/Developer/HotReload/Private/HotReloadClassReinstancer.cpp
@@ -228,6 +228,8 @@ void FHotReloadClassReinstancer::RecreateCDOAndSetupOldClassReinstancing(UClass*
 	// Re-create the CDO, re-running its constructor
 	ReconstructClassDefaultObject(InOldClass, CDOOuter, CDOName, CDOFlags);
 
+	ReconstructedCDOsMap.Add(OriginalCDO, InOldClass->GetDefaultObject());
+
 	// Collect the property values after re-constructing the CDO
 	SerializeCDOProperties(InOldClass->GetDefaultObject(), ReconstructedCDOProperties);
 
@@ -260,10 +262,11 @@ void FHotReloadClassReinstancer::RecreateCDOAndSetupOldClassReinstancing(UClass*
 	}
 }
 
-FHotReloadClassReinstancer::FHotReloadClassReinstancer(UClass* InNewClass, UClass* InOldClass)
+FHotReloadClassReinstancer::FHotReloadClassReinstancer(UClass* InNewClass, UClass* InOldClass, TMap<UObject*, UObject*>& OutReconstructedCDOsMap)
 	: NewClass(nullptr)
 	, bNeedsReinstancing(false)
 	, CopyOfPreviousCDO(nullptr)
+	, ReconstructedCDOsMap(OutReconstructedCDOsMap)
 {
 	ensure(InOldClass);
 	ensure(!HotReloadedOldClass && !HotReloadedNewClass);
diff --git a/Engine/Source/Developer/HotReload/Private/HotReloadClassReinstancer.h b/Engine/Source/Developer/HotReload/Private/HotReloadClassReinstancer.h
--- a/Engine/Source/Developer/HotReload/Private/HotReloadClassReinstancer.h
+++ b/Engine/Source/Developer/HotReload/Private/HotReloadClassReinstancer.h
@@ -94,7 +94,7 @@ class FHotReloadClassReinstancer : public FBlueprintCompileReinstancer
 public:
 
 	/** Sets the re-instancer up to re-instance native classes */
-	FHotReloadClassReinstancer(UClass* InNewClass, UClass* InOldClass);
+	FHotReloadClassReinstancer(UClass* InNewClass, UClass* InOldClass, TMap<UObject*, UObject*>& OutReconstructedCDOsMap);
 	
 	/** Destructor */
 	virtual ~FHotReloadClassReinstancer();
@@ -109,9 +109,9 @@ public:
 	void ReinstanceObjectsAndUpdateDefaults();
 
 	/** Creates the reinstancer as a sharable object */
-	static TSharedPtr<FHotReloadClassReinstancer> Create(UClass* InNewClass, UClass* InOldClass)
+	static TSharedPtr<FHotReloadClassReinstancer> Create(UClass* InNewClass, UClass* InOldClass, TMap<UObject*, UObject*>& OutReconstructedCDOsMap)
 	{
-		return MakeShareable(new FHotReloadClassReinstancer(InNewClass, InOldClass));
+		return MakeShareable(new FHotReloadClassReinstancer(InNewClass, InOldClass, OutReconstructedCDOsMap));
 	}
 
 	// FSerializableObject interface
@@ -125,6 +125,10 @@ protected:
 	// FBlueprintCompileReinstancer interface
 	virtual bool ShouldPreserveRootComponentOfReinstancedActor() const override { return false; }
 	// End of FBlueprintCompileReinstancer interface
+
+private:
+	/** Reference to reconstructed CDOs map in this hot-reload session. */
+	TMap<UObject*, UObject*>& ReconstructedCDOsMap;
 };
 
 #endif // WITH_ENGINE
\ No newline at end of file
