HotReload: Checking in a speculative fix for CIS issues

diff --git a/Engine/Source/Developer/HotReload/Private/HotReload.cpp b/Engine/Source/Developer/HotReload/Private/HotReload.cpp
--- a/Engine/Source/Developer/HotReload/Private/HotReload.cpp
+++ b/Engine/Source/Developer/HotReload/Private/HotReload.cpp
@@ -7,6 +7,7 @@
 #if WITH_ENGINE
 #include "HotReloadClassReinstancer.h"
 #include "EngineAnalytics.h"
+#include "Runtime/Engine/Classes/Engine/BlueprintGeneratedClass.h"
 #endif
 
 DEFINE_LOG_CATEGORY(LogHotReload);
@@ -814,7 +815,12 @@ void FHotReloadModule::ReplaceReferencesToReconstructedCDOs()
 
 		for (auto& ReferencingProperty : ReferencingProperties)
 		{
-			static const UProperty* PropertyToSkip = UBlueprintGeneratedClass::StaticClass()->FindPropertyByName(TEXT("OverridenArchetypeForCDO"));
+			static const UProperty* PropertyToSkip =
+#if WITH_ENGINE
+				UBlueprintGeneratedClass::StaticClass()->FindPropertyByName(GET_MEMBER_NAME_CHECKED(UBlueprintGeneratedClass, OverridenArchetypeForCDO));
+#else
+				nullptr;
+#endif
 			if (!ReferencingProperty.Value->IsA<UObjectProperty>() || ReferencingProperty.Value == PropertyToSkip)
 			{
 				continue;
