UE-5946: Cannot compile/save BP after hot reload

diff --git a/Engine/Source/Developer/HotReload/Private/HotReload.cpp b/Engine/Source/Developer/HotReload/Private/HotReload.cpp
--- a/Engine/Source/Developer/HotReload/Private/HotReload.cpp
+++ b/Engine/Source/Developer/HotReload/Private/HotReload.cpp
@@ -9,6 +9,7 @@
 #include "EngineAnalytics.h"
 #include "Runtime/Engine/Classes/Engine/BlueprintGeneratedClass.h"
 #endif
+#include "KismetEditorUtilities.h"
 
 DEFINE_LOG_CATEGORY(LogHotReload);
 
@@ -159,6 +160,11 @@ private:
 	 */
 	void ReplaceReferencesToReconstructedCDOs();
 
+	/**
+	 * Recompiles blueprints that were touched during this hot-reload.
+	 */
+	void RecompileTouchedBlueprints();
+
 	/**
 	* Callback for async ompilation
 	*/
@@ -588,6 +594,9 @@ inline uint32 GetTypeHash(Native A)
 /** Map from old function pointer to new function pointer for hot reload. */
 static TMap<Native, Native> HotReloadFunctionRemap;
 
+static TSet<UBlueprint*> HotReloadBPSetToRecompile;
+static TSet<UBlueprint*> HotReloadBPSetToRecompileBytecodeOnly;
+
 /** Adds and entry for the UFunction native pointer remap table */
 void FHotReloadModule::AddHotReloadFunctionRemap(Native NewFunctionPointer, Native OldFunctionPointer)
 {
@@ -766,6 +775,9 @@ ECompilationResult::Type FHotReloadModule::DoHotReloadInternal(bool bRecompileFi
 			}
 			HotReloadAr.Logf(ELogVerbosity::Warning, TEXT("HotReload successful (%d functions remapped  %d scriptstructs remapped)"), Count, ScriptStructs.Num());
 
+			RecompileTouchedBlueprints();
+
+
 			HotReloadFunctionRemap.Empty();
 
 			ReplaceReferencesToReconstructedCDOs();
@@ -836,6 +848,43 @@ void FHotReloadModule::ReplaceReferencesToReconstructedCDOs()
 	ReconstructedCDOsMap.Empty();
 }
 
+void FHotReloadModule::RecompileTouchedBlueprints()
+{
+	bool bCollectGarbage = HotReloadBPSetToRecompile.Num() > 0;
+
+	while (HotReloadBPSetToRecompile.Num())
+	{
+		auto Iter = HotReloadBPSetToRecompile.CreateIterator();
+		TWeakObjectPtr<UBlueprint> BPPtr = *Iter;
+		Iter.RemoveCurrent();
+		if (auto BP = BPPtr.Get())
+		{
+			FKismetEditorUtilities::CompileBlueprint(BP, false, true);
+		}
+	}
+
+	if (bCollectGarbage)
+	{
+		// Garbage collect to make sure the old class and actors are disposed of
+		CollectGarbage(GARBAGE_COLLECTION_KEEPFLAGS);
+	}
+
+	HotReloadBPSetToRecompile.Empty();
+
+	while (HotReloadBPSetToRecompileBytecodeOnly.Num())
+	{
+		auto Iter = HotReloadBPSetToRecompileBytecodeOnly.CreateIterator();
+		TWeakObjectPtr<UBlueprint> BPPtr = *Iter;
+		Iter.RemoveCurrent();
+		if (auto BP = BPPtr.Get())
+		{
+			FKismetEditorUtilities::RecompileBlueprintBytecode(BP);
+		}
+	}
+
+	HotReloadBPSetToRecompileBytecodeOnly.Empty();
+}
+
 ECompilationResult::Type FHotReloadModule::RebindPackages(TArray<UPackage*> InPackages, TArray<FName> DependentModules, const bool bWaitForCompletion, FOutputDevice &Ar)
 {
 	ECompilationResult::Type Result = ECompilationResult::Unknown;
@@ -935,7 +984,7 @@ ECompilationResult::Type FHotReloadModule::RebindPackagesInternal(TArray<UPackag
 #if WITH_ENGINE
 void FHotReloadModule::ReinstanceClass(UClass* OldClass, UClass* NewClass)
 {	
-	auto ReinstanceHelper = FHotReloadClassReinstancer::Create(NewClass, OldClass, ReconstructedCDOsMap);
+	auto ReinstanceHelper = FHotReloadClassReinstancer::Create(NewClass, OldClass, ReconstructedCDOsMap, HotReloadBPSetToRecompile, HotReloadBPSetToRecompileBytecodeOnly);
 	if (ReinstanceHelper->ClassNeedsReinstancing())
 	{
 		UE_LOG(LogHotReload, Log, TEXT("Re-instancing %s after hot-reload."), NewClass ? *NewClass->GetName() : *OldClass->GetName());
diff --git a/Engine/Source/Developer/HotReload/Private/HotReloadClassReinstancer.cpp b/Engine/Source/Developer/HotReload/Private/HotReloadClassReinstancer.cpp
--- a/Engine/Source/Developer/HotReload/Private/HotReloadClassReinstancer.cpp
+++ b/Engine/Source/Developer/HotReload/Private/HotReloadClassReinstancer.cpp
@@ -262,11 +262,13 @@ void FHotReloadClassReinstancer::RecreateCDOAndSetupOldClassReinstancing(UClass*
 	}
 }
 
-FHotReloadClassReinstancer::FHotReloadClassReinstancer(UClass* InNewClass, UClass* InOldClass, TMap<UObject*, UObject*>& OutReconstructedCDOsMap)
+FHotReloadClassReinstancer::FHotReloadClassReinstancer(UClass* InNewClass, UClass* InOldClass, TMap<UObject*, UObject*>& OutReconstructedCDOsMap, TSet<UBlueprint*>& InBPSetToRecompile, TSet<UBlueprint*>& InBPSetToRecompileBytecodeOnly)
 	: NewClass(nullptr)
 	, bNeedsReinstancing(false)
 	, CopyOfPreviousCDO(nullptr)
 	, ReconstructedCDOsMap(OutReconstructedCDOsMap)
+	, BPSetToRecompile(InBPSetToRecompile)
+	, BPSetToRecompileBytecodeOnly(InBPSetToRecompileBytecodeOnly)
 {
 	ensure(InOldClass);
 	ensure(!HotReloadedOldClass && !HotReloadedNewClass);
@@ -496,4 +498,38 @@ void FHotReloadClassReinstancer::AddReferencedObjects(FReferenceCollector& Colle
 	Collector.AddReferencedObject(CopyOfPreviousCDO);
 }
 
+void FHotReloadClassReinstancer::EnlistDependentBlueprintToRecompile(UBlueprint* BP, bool bBytecodeOnly)
+{
+	if (IsValid(BP))
+	{
+		if (bBytecodeOnly)
+		{
+			if (!BPSetToRecompile.Contains(BP) && !BPSetToRecompileBytecodeOnly.Contains(BP))
+			{
+				BPSetToRecompileBytecodeOnly.Add(BP);
+			}
+		}
+		else
+		{
+			if (!BPSetToRecompile.Contains(BP))
+			{
+				if (BPSetToRecompileBytecodeOnly.Contains(BP))
+				{
+					BPSetToRecompileBytecodeOnly.Remove(BP);
+				}
+
+				BPSetToRecompile.Add(BP);
+			}
+		}
+	}
+}
+
+void FHotReloadClassReinstancer::BlueprintWasRecompiled(UBlueprint* BP, bool bBytecodeOnly)
+{
+	BPSetToRecompile.Remove(BP);
+	BPSetToRecompileBytecodeOnly.Remove(BP);
+
+	FBlueprintCompileReinstancer::BlueprintWasRecompiled(BP, bBytecodeOnly);
+}
+
 #endif
diff --git a/Engine/Source/Developer/HotReload/Private/HotReloadClassReinstancer.h b/Engine/Source/Developer/HotReload/Private/HotReloadClassReinstancer.h
--- a/Engine/Source/Developer/HotReload/Private/HotReloadClassReinstancer.h
+++ b/Engine/Source/Developer/HotReload/Private/HotReloadClassReinstancer.h
@@ -94,7 +94,7 @@ class FHotReloadClassReinstancer : public FBlueprintCompileReinstancer
 public:
 
 	/** Sets the re-instancer up to re-instance native classes */
-	FHotReloadClassReinstancer(UClass* InNewClass, UClass* InOldClass, TMap<UObject*, UObject*>& OutReconstructedCDOsMap);
+	FHotReloadClassReinstancer(UClass* InNewClass, UClass* InOldClass, TMap<UObject*, UObject*>& OutReconstructedCDOsMap, TSet<UBlueprint*>& InBPSetToRecompile, TSet<UBlueprint*>& InBPSetToRecompileBytecodeOnly);
 	
 	/** Destructor */
 	virtual ~FHotReloadClassReinstancer();
@@ -109,9 +109,9 @@ public:
 	void ReinstanceObjectsAndUpdateDefaults();
 
 	/** Creates the reinstancer as a sharable object */
-	static TSharedPtr<FHotReloadClassReinstancer> Create(UClass* InNewClass, UClass* InOldClass, TMap<UObject*, UObject*>& OutReconstructedCDOsMap)
+	static TSharedPtr<FHotReloadClassReinstancer> Create(UClass* InNewClass, UClass* InOldClass, TMap<UObject*, UObject*>& OutReconstructedCDOsMap, TSet<UBlueprint*>& InBPSetToRecompile, TSet<UBlueprint*>& InBPSetToRecompileBytecodeOnly)
 	{
-		return MakeShareable(new FHotReloadClassReinstancer(InNewClass, InOldClass, OutReconstructedCDOsMap));
+		return MakeShareable(new FHotReloadClassReinstancer(InNewClass, InOldClass, OutReconstructedCDOsMap, InBPSetToRecompile, InBPSetToRecompileBytecodeOnly));
 	}
 
 	// FSerializableObject interface
@@ -120,6 +120,9 @@ public:
 
 	virtual bool IsClassObjectReplaced() const override  { return true; }
 
+	virtual void EnlistDependentBlueprintToRecompile(UBlueprint* BP, bool bBytecodeOnly) override;
+	virtual void BlueprintWasRecompiled(UBlueprint* BP, bool bBytecodeOnly) override;
+
 protected:
 
 	// FBlueprintCompileReinstancer interface
@@ -129,6 +132,9 @@ protected:
 private:
 	/** Reference to reconstructed CDOs map in this hot-reload session. */
 	TMap<UObject*, UObject*>& ReconstructedCDOsMap;
+	TSet<UBlueprint*>& BPSetToRecompile;
+	TSet<UBlueprint*>& BPSetToRecompileBytecodeOnly;
+
 };
 
 #endif // WITH_ENGINE
\ No newline at end of file
diff --git a/Engine/Source/Editor/UnrealEd/Public/Kismet2/KismetReinstanceUtilities.h b/Engine/Source/Editor/UnrealEd/Public/Kismet2/KismetReinstanceUtilities.h
--- a/Engine/Source/Editor/UnrealEd/Public/Kismet2/KismetReinstanceUtilities.h
+++ b/Engine/Source/Editor/UnrealEd/Public/Kismet2/KismetReinstanceUtilities.h
@@ -75,8 +75,8 @@ public:
 	static void OptionallyRefreshNodes(UBlueprint* BP);
 
 	void ListDependentBlueprintsToRefresh(const TArray<UBlueprint*>& DependentBPs);
-	void EnlistDependentBlueprintToRecompile(UBlueprint* BP, bool bBytecodeOnly);
-	void BlueprintWasRecompiled(UBlueprint* BP, bool bBytecodeOnly);
+	virtual void EnlistDependentBlueprintToRecompile(UBlueprint* BP, bool bBytecodeOnly);
+	virtual void BlueprintWasRecompiled(UBlueprint* BP, bool bBytecodeOnly);
 
 	static TSharedPtr<FBlueprintCompileReinstancer> Create(UClass* InClassToReinstance, bool bIsBytecodeOnly = false, bool bSkipGC = false)
 	{
