CIS fix. PortalPrototype needs some special compilation handling for hot-reload (#if WITH_ENGINE). Fixed.

diff --git a/Engine/Source/Developer/HotReload/Private/HotReload.cpp b/Engine/Source/Developer/HotReload/Private/HotReload.cpp
--- a/Engine/Source/Developer/HotReload/Private/HotReload.cpp
+++ b/Engine/Source/Developer/HotReload/Private/HotReload.cpp
@@ -8,8 +8,8 @@
 #include "HotReloadClassReinstancer.h"
 #include "EngineAnalytics.h"
 #include "Runtime/Engine/Classes/Engine/BlueprintGeneratedClass.h"
-#endif
 #include "KismetEditorUtilities.h"
+#endif
 
 DEFINE_LOG_CATEGORY(LogHotReload);
 
@@ -160,10 +160,12 @@ private:
 	 */
 	void ReplaceReferencesToReconstructedCDOs();
 
+#if WITH_ENGINE
 	/**
 	 * Recompiles blueprints that were touched during this hot-reload.
 	 */
 	void RecompileTouchedBlueprints();
+#endif // WITH_ENGINE
 
 	/**
 	* Callback for async ompilation
@@ -775,8 +777,9 @@ ECompilationResult::Type FHotReloadModule::DoHotReloadInternal(bool bRecompileFi
 			}
 			HotReloadAr.Logf(ELogVerbosity::Warning, TEXT("HotReload successful (%d functions remapped  %d scriptstructs remapped)"), Count, ScriptStructs.Num());
 
+#if WITH_ENGINE
 			RecompileTouchedBlueprints();
-
+#endif // WITH_ENGINE
 
 			HotReloadFunctionRemap.Empty();
 
@@ -848,6 +851,7 @@ void FHotReloadModule::ReplaceReferencesToReconstructedCDOs()
 	ReconstructedCDOsMap.Empty();
 }
 
+#if WITH_ENGINE
 void FHotReloadModule::RecompileTouchedBlueprints()
 {
 	bool bCollectGarbage = HotReloadBPSetToRecompile.Num() > 0;
@@ -884,6 +888,7 @@ void FHotReloadModule::RecompileTouchedBlueprints()
 
 	HotReloadBPSetToRecompileBytecodeOnly.Empty();
 }
+#endif // WITH_ENGINE
 
 ECompilationResult::Type FHotReloadModule::RebindPackages(TArray<UPackage*> InPackages, TArray<FName> DependentModules, const bool bWaitForCompletion, FOutputDevice &Ar)
 {
