CIS fix for CL2710850 for other platforms.

diff --git a/Engine/Source/Developer/HotReload/Private/HotReload.cpp b/Engine/Source/Developer/HotReload/Private/HotReload.cpp
--- a/Engine/Source/Developer/HotReload/Private/HotReload.cpp
+++ b/Engine/Source/Developer/HotReload/Private/HotReload.cpp
@@ -11,8 +11,6 @@
 #include "KismetEditorUtilities.h"
 #endif
 
-#include "KismetEditorUtilities.h"
-
 DEFINE_LOG_CATEGORY(LogHotReload);
 
 #define LOCTEXT_NAMESPACE "HotReload"
diff --git a/Engine/Source/Editor/UnrealEd/Public/Layers/ILayers.h b/Engine/Source/Editor/UnrealEd/Public/Layers/ILayers.h
--- a/Engine/Source/Editor/UnrealEd/Public/Layers/ILayers.h
+++ b/Engine/Source/Editor/UnrealEd/Public/Layers/ILayers.h
@@ -3,6 +3,8 @@
 
 #include "IFilter.h"
 
+class FLevelEditorViewportClient;
+
 namespace ELayersAction
 {
 	enum Type
