Copying //UE4/Dev-Core to //UE4/Dev-Main (Source: //UE4/Dev-Core @ 3049602)

==========================
MAJOR FEATURES + CHANGES
==========================

Change 2946506 on 2016/04/18 by [REDACTED_PERSON]

	Update to Crash Reporter buggs table to add new search fields and inclusion of packages needed for e-mail reports.

Change 3017807 on 2016/06/17 by [REDACTED_PERSON]

	Improved Crash Report Process folder delete code as it could sometimes fail.
	[UE-30349] - Crash Report Process is leaving crashes in the landing zone that build up and block the queue

	Also added logging to Slack when stop request received instead of just when stop is complete.

Change 3019367 on 2016/06/20 by [REDACTED_PERSON]

	Improve Crash Report Process logging to track bad reads from S3.
	Also, better logging when CleanReport fails to delete folders.

Change 3019376 on 2016/06/20 by [REDACTED_PERSON]

	Clarification of assert message and comments which talk about 'null' TFunctions.
	Tidy-up of dead code.

Change 3019409 on 2016/06/20 by [REDACTED_PERSON]

	New Find and FindByPredicate algorithms for finding stuff in arbitrary containers.

Change 3022658 on 2016/06/22 by [REDACTED_PERSON]

	Discarding duplicated crash reports earlier in read from Data Router process to avoid clashes in the landing zone (CRP v1.1.11)
	[UE-30349] - Crash Report Process is leaving crashes in the landing zone that build up and block the queue

	Also improved logging to Slack with better layout, fixed event ordering and counting duplicates.

Change 3022840 on 2016/06/22 by [REDACTED_PERSON]

	Skipped UHT attributes removed.

Change 3022907 on 2016/06/22 by [REDACTED_PERSON]

	Fixing crash when adding a new C++ class to project

Change 3023169 on 2016/06/22 by [REDACTED_PERSON]

	Checks for UTHINGs in skipped preprocessor blocks.
	Fixes for skipped UTHINGs and some other parsing accidents.

Change 3023239 on 2016/06/22 by [REDACTED_PERSON]

	Fix for JSON date parsing reported here: https://udn.unrealengine.com/questions/299342/fdatetime-json-serialization-bug.html

Change 3026812 on 2016/06/24 by [REDACTED_PERSON]

	Marked FEnvQueryInstance::AddItemData UEnvQueryItemType_Point specialization as AIMODULE_API #UE4

Change 3028235 on 2016/06/27 by [REDACTED_PERSON]

	PR #2535: BUGFIX: FPS pop-up updates when loading new stat file (Contributed by [REDACTED_PERSON])

Change 3028282 on 2016/06/27 by [REDACTED_PERSON]

	Fix for missing UFUNCTION check in skipped preprocessor blocks.

Change 3028284 on 2016/06/27 by [REDACTED_PERSON]

	Debuggability improvements and coding standards changes.

Change 3028343 on 2016/06/27 by [REDACTED_PERSON]

	Fix for UHT error in WEX.

Change 3028393 on 2016/06/27 by [REDACTED_PERSON]

	Fix for hot reload of enums finding the old enum.
	Fix to stop SPropertyEditorNumeric caching the enum flags.

Change 3030362 on 2016/06/28 by [REDACTED_PERSON]

	Fixing hang when cooking.

Change 3030462 on 2016/06/28 by [REDACTED_PERSON]

	Assert added to PackageTools::GetFilteredPackageList() to help with catching a bug reported in the wild.

Change 3034341 on 2016/06/30 by [REDACTED_PERSON]

	Modified crash handling code (on Windows) to handle two threads crashing at the same time properly.  Previously the second crash would force the process to exit before generating the crash report.

	Added 'debug twothreadsgpf' command to test the functionality.

Change 3034342 on 2016/06/30 by [REDACTED_PERSON]

	Fix for crash when loading an empty cached asset registry.

Change 3035599 on 2016/07/01 by [REDACTED_PERSON]

	Added support for CrashType string to Crash Report Process. CRP v1.1.12
	[UE-30592] - Crash Reporter should determine crash type on client and pass string to server

	Also fixes problem with reports falling back on the legacy WER metadata when a crash context exists. They now only read the error message from metadata if available and keep crash context data when possible.

	Added in missing crash context parameters that have been added to clients but not known by the server.

Change 3035787 on 2016/07/01 by [REDACTED_PERSON]

	Fix for crash when DuplicateRedirects does not contain the DependentObject when saving dependencies. It will still fall through to the assertion below, but it will now fail with a useful error message instead of a generic 'Pair != nullptr' from Map.h.

Change 3036933 on 2016/07/04 by [REDACTED_PERSON]

	Proper forwarding constructor for FAsyncTask.

Change 3036938 on 2016/07/04 by [REDACTED_PERSON]

	Fix for CDO hot reload corrupting memory when replacing references inside structs.

Change 3036960 on 2016/07/04 by [REDACTED_PERSON]

	Fix for FAnsiAllocator::ResizeAllocation when resizing to zero.

Change 3037423 on 2016/07/05 by [REDACTED_PERSON]

	FModuleManager::UnloadOrAbandonModuleWithCallback split into two instead of switching behavior with a bool.

Change 3037464 on 2016/07/05 by [REDACTED_PERSON]

	HotReload.cpp cleanup:

	Deep nesting flattened.
	Linear array searches replaced with maps.
	FHotReloadModule::GetGameModules made into a non-member function and split into two.
	Comment and coding standard fixes.

Change 3037741 on 2016/07/05 by [REDACTED_PERSON]

	Fix for COTF not checking the correct timestamps on startup.

Change 3037846 on 2016/07/05 by [REDACTED_PERSON]

	Fix for compile button disappearing on a bad compile.

Change 3037994 on 2016/07/05 by [REDACTED_PERSON]

	Static analysis fixes:

	warning C6308: 'realloc' might return null pointer: assigning null pointer to 'Data', which is passed as an argument to 'realloc', will cause the original memory block to be leaked.

Change 3039186 on 2016/07/06 by [REDACTED_PERSON]

	Enabling crash callstack logging by default.

Change 3039220 on 2016/07/06 by [REDACTED_PERSON]

	Static analysis fixes:

	warning C28159: Consider using 'InitiateSystemShutdownEx' instead of 'ExitWindowsEx'. Reason: Legacy API. Rearchitect to avoid Reboot
	warning C6001: Using uninitialized memory 'UserNameLength'
	warning C6001: Using uninitialized memory 'DomainNameLength'

Change 3039230 on 2016/07/06 by [REDACTED_PERSON]

	Fix for VC internal compiler errors.

Change 3039237 on 2016/07/06 by [REDACTED_PERSON]

	Static analysis fix: warning C6385: Reading invalid data from 'Path':  the readable size is '400' bytes, but 'PathCurrentDepth' bytes may be read.

Change 3039287 on 2016/07/06 by [REDACTED_PERSON]

	Static analysis fixes:

	warning C6509: Invalid annotation: 'return' cannot be referenced in some contexts
	warning C6101: Returning uninitialized memory '*lpdwExitCode'.  A successful path through the function does not set the named _Out_ parameter.
	warning C6387: '_Param_(1)' could be '0':  this does not adhere to the specification for the function 'IMoniker::BindToStorage'.
	warning C6387: '_Param_(1)' could be '0':  this does not adhere to the specification for the function 'IMoniker::BindToObject'.
	warning C6031: Return value ignored: 'CoCreateInstance'.

Change 3039359 on 2016/07/06 by [REDACTED_PERSON]

	Compile error fix for FAsyncTask, courtesy of SteveR

Change 3039534 on 2016/07/06 by [REDACTED_PERSON]

	Static analysis fix: warning C6319: Use of the comma-operator in a tested expression causes the left argument to be ignored when it has no side-effects.

Change 3039545 on 2016/07/06 by [REDACTED_PERSON]

	Static analysis fix: warning C6297: Arithmetic overflow:  32-bit value is shifted, then cast to 64-bit value.  Results might not be an expected value.

Change 3039578 on 2016/07/06 by [REDACTED_PERSON]

	Static analysis fix: warning C6263: Using _alloca in a loop:  this can quickly overflow stack.

Change 3039623 on 2016/07/06 by [REDACTED_PERSON]

	Static analysis fixes:

	warning C6011: Dereferencing NULL pointer 'X'
	warning C6308:'realloc' might return null pointer: assigning null pointer to 'X', which is passed as an argument to 'realloc', will cause the original memory block to be leaked.
	warning C6385: Reading invalid data from 'X':  the readable size is 'Y' bytes, but 'Z' bytes may be read.
	warning C6386: Buffer overrun while writing to 'X':  the writable size is 'Y' bytes, but 'Z' bytes might be written.
	warning C28182: Dereferencing NULL pointer. 'X' contains the same NULL value as 'Y' did.

Change 3039630 on 2016/07/06 by [REDACTED_PERSON]

	Fix for crash when spawning an actor using a template object that has instance components. UActorComponent::PostInitProperties was adding itself to the owner's InstanceComponents array, resulting in a realloc of that array and invalidating the reference that the owner's ObjectInitializer was trying to replace while instantiating that property. The new instance component will be added to the array as part of the owner's initialization anyway, so it is not necessary to do it here.

Change 3039664 on 2016/07/06 by [REDACTED_PERSON]

	Static analysis fixes:

	warning C6386: Buffer overrun while writing to 'NewKeys':  the writable size is 'NewIndexSize*4' bytes, but '8' bytes might be written.
	warning C6386: Buffer overrun while writing to 'NewHeapIndexes':  the writable size is 'NewIndexSize*4' bytes, but '8' bytes might be written.

Change 3039673 on 2016/07/06 by [REDACTED_PERSON]

	Static analysis fix: warning C6011: Dereferencing NULL pointer 'v'.

Change 3039690 on 2016/07/06 by [REDACTED_PERSON]

	Static analysis fixes:

	warning C6011: Dereferencing NULL pointer 'X'.
	warning C6246: Local declaration of 'X' hides declaration of the same name in outer scope.
	warning C6262: Function uses '121180' bytes of stack:  exceeds /analyze:stacksize '81940'.  Consider moving some data to heap.
	warning C6263: Using _alloca in a loop:  this can quickly overflow stack.

Change 3040868 on 2016/07/07 by [REDACTED_PERSON]

	Config based class stripping for server builds

Change 3040872 on 2016/07/07 by [REDACTED_PERSON]

	Remove "return false" NeedsLoadForServer functions from engine code

Change 3040997 on 2016/07/07 by [REDACTED_PERSON]

	Static analysis fixes:

	warning C6011: Dereferencing NULL pointer 'Landscape'.
	warning C6011: Dereferencing NULL pointer 'rhs.Allocation.LayerInfo'.
	warning C6011: Dereferencing NULL pointer 'lhs.Allocation.LayerInfo'.

Change 3041004 on 2016/07/07 by [REDACTED_PERSON]

	Static analysis fix: warning C6336: Arithmetic operator has precedence over question operator, use parentheses to clarify intent.

Change 3041014 on 2016/07/07 by [REDACTED_PERSON]

	Static analysis fix: warning C6287: Redundant code:  the left and right sub-expressions are identical.

Change 3041111 on 2016/07/07 by [REDACTED_PERSON]

	Removal of an obsolete error message about INI file case sensitivity.

Change 3041150 on 2016/07/07 by [REDACTED_PERSON]

	Static analysis fix: warning C6289: Incorrect operator:  mutual exclusion over || is always a non-zero constant.  Did you intend to use && instead?

Change 3041274 on 2016/07/07 by [REDACTED_PERSON]

	Static analysis fixes: warning C6001: Using uninitialized memory 'X'.

Change 3041294 on 2016/07/07 by [REDACTED_PERSON]

	Fixed protocol buffer and decompression errors in Crash Report Process (v.1.1.14)
	[UE-32151] - High number of crashes read from S3 by [REDACTED_PERSON] are failing to unpack

	Size of buffer received from S3 is incorrect for some records. Fixed read problems by using size header value instead of stream length.

	Increased buffer size for decompression as this was sometimes too small.

	Modified S3 reading code to look for multiple records in each downloaded file.

Change 3041472 on 2016/07/07 by [REDACTED_PERSON]

	Static analysis fixes:

	warning C6294: Ill-defined for-loop:  initial condition does not satisfy test.  Loop body not executed.
	warning C6201: Index '1' is out of valid index range '0' to '0' for possibly stack allocated buffer 'NewHistory.Nodes'.

Change 3043074 on 2016/07/08 by [REDACTED_PERSON]

	Fix for COTF incorrectly reconstructing the original asset path based on the sandbox path when the game name differs from the game folder name.
	Fix for COTF GetFiles not handling absolute GameDir paths properly.

Change 3044461 on 2016/07/11 by [REDACTED_PERSON]

	Static analysis fix: warning C6386: Buffer overrun while writing to 'Attributes':  the writable size is '16384' bytes, but '-8' bytes might be written.

Change 3044470 on 2016/07/11 by [REDACTED_PERSON]

	Static analysis fix: warning C6011: Dereferencing NULL pointer 'Node.Sequence'.

Change 3044476 on 2016/07/11 by [REDACTED_PERSON]

	Static analysis fix: warning C6011: Dereferencing NULL pointer 'Property'.

Change 3044551 on 2016/07/11 by [REDACTED_PERSON]

	Static analysis fix: warning C28182: Dereferencing NULL pointer. 'Node' contains the same NULL value as 'KeyAreaNode' did.

Change 3044664 on 2016/07/11 by [REDACTED_PERSON]

	Static analysis fixes:

	warning C6011: Dereferencing NULL pointer 'ToLandscape->SplineComponent'.
	warning C28182: Dereferencing NULL pointer. 'SplinesComponent' contains the same NULL value as 'Landscape->SplineComponent' did.
	warning C6011: Dereferencing NULL pointer 'Landscape->SplineComponent'.
	warning C6385: Reading invalid data from 'out':  the readable size is 'sizeof(kiss_fft_cpx)*Dims[0]*Dims[1]' bytes, but '16' bytes may be read.

Change 3044716 on 2016/07/11 by [REDACTED_PERSON]

	Static analysis fix: warning C6385: Reading invalid data from 'this->ScreenSize':  the readable size is '32' bytes, but '-4' bytes may be read.

Change 3044717 on 2016/07/11 by [REDACTED_PERSON]

	Static analysis fix: warning C28182: Dereferencing NULL pointer. 'Window' contains the same NULL value as 'ElementType * Window=AllWindows.FindByPredicate((*FStaticMeshEditorTest::RunTest::<lambda_46fd0093f3912289e870263afe1fcb2e>(ExpectedTitle)))' did.

	This appears to be a false positive.

Change 3044787 on 2016/07/11 by [REDACTED_PERSON]

	Static analysis fixes:

	warning C6011: Dereferencing NULL pointer 'FbxObject'.
	warning C28182: Dereferencing NULL pointer. 'Node' contains the same NULL value as 'RigidMeshNode' did.
	warning C28182: Dereferencing NULL pointer. 'Node' contains the same NULL value as 'Result' did.

Change 3045933 on 2016/07/12 by [REDACTED_PERSON]

	Overloading support for TSharedPtr, TSharedRef and TWeakPtr.

Change 3045960 on 2016/07/12 by [REDACTED_PERSON]

	Fixing a crash in Portal (and any other program that uses UObjects and GCs, with the exception of UHT) caused by classes not having their token stream assembled.

Change 3045963 on 2016/07/12 by [REDACTED_PERSON]

	PLATFORM_COMPILER_HAS_EXPLICIT_OPERATORS, FORCEINLINE_EXPLICIT_OPERATOR_BOOL and SAFE_BOOL_OPERATORS macros removed.
	THasOperatorEquals and THasOperatorNotEquals traits moved to their own header.

Change 3045967 on 2016/07/12 by [REDACTED_PERSON]

	Initializer list support for TArray and TSet.

Change 3045968 on 2016/07/12 by [REDACTED_PERSON]

	Fixing an ensure after typing 'stat dumphitches' in console.

Change 3045992 on 2016/07/12 by [REDACTED_PERSON]

	Making sure CoreUObject headers are included for programs that don't include the engine (fixing MinidumpDiagnostics CIS failure)

Change 3047870 on 2016/07/13 by [REDACTED_PERSON]

	Updated CRW to entity framework with repository models. #rb none

Change 3047871 on 2016/07/13 by [REDACTED_PERSON]

	Add repository models #rb none

Change 3049468 on 2016/07/14 by [REDACTED_PERSON]

	Fix broken project files. #rb none

diff --git a/Engine/Config/BaseEngine.ini b/Engine/Config/BaseEngine.ini
--- a/Engine/Config/BaseEngine.ini
+++ b/Engine/Config/BaseEngine.ini
@@ -1591,6 +1591,23 @@ CanSkipEditorReferencedPackagesWhenCooking=False
 DetailedCallstacksInNonMonolithicBuilds=False
 UseSeperateBulkDataFiles=False
 HangDuration=0.0
++ClassesExcludedForServer=WidgetBlueprint
++ClassesExcludedForServer=GroupActor
++ClassesExcludedForServer=MetaData
++ClassesExcludedForServer=ObjectRedirector
++ClassesExcludedForServer=NavMeshRenderingComponent
++ClassesExcludedForServer=ReflectionCaptureComponent
++ClassesExcludedForServer=TextRenderComponent
++ClassesExcludedForServer=Font
++ClassesExcludedForServer=InterpCurveEdSetup
++ClassesExcludedForServer=MaterialExpression
++ClassesExcludedForServer=MatineeActorCameraAnim
++ClassesExcludedForServer=ParticleEmitter
++ClassesExcludedForServer=ParticleLODLevel
++ClassesExcludedForServer=ParticleModule
++ClassesExcludedForServer=SubUVAnimation
++ClassesExcludedForServer=SoundNode
++ClassesExcludedForServer=GameplayEffectUIData
 
 [/Script/Engine.StreamingSettings]
 s.AsyncIOBandwidthLimit=0
diff --git a/Engine/Source/Developer/GameplayDebugger/Classes/GameplayDebuggerSettings.h b/Engine/Source/Developer/GameplayDebugger/Classes/GameplayDebuggerSettings.h
--- a/Engine/Source/Developer/GameplayDebugger/Classes/GameplayDebuggerSettings.h
+++ b/Engine/Source/Developer/GameplayDebugger/Classes/GameplayDebuggerSettings.h
@@ -69,7 +69,6 @@ public:
 
 #if ADD_LEVEL_EDITOR_EXTENSIONS
 	/** Experimental setting to extend viewport menu in Simulate, to have quick access to GameplayDebugger settings. */
-	UPROPERTY(EditAnywhere, config, Category = "GameplayDebugger")
 	bool bExtendViewportMenu;
 #endif
 
diff --git a/Engine/Source/Developer/HotReload/Private/HotReload.cpp b/Engine/Source/Developer/HotReload/Private/HotReload.cpp
--- a/Engine/Source/Developer/HotReload/Private/HotReload.cpp
+++ b/Engine/Source/Developer/HotReload/Private/HotReload.cpp
@@ -129,17 +129,6 @@ private:
 		{ }
 	};
 
-	struct FRecompiledModule
-	{
-		FString Name;
-		FString NewFilename;
-		FRecompiledModule() {}
-		FRecompiledModule(FString InName, FString InFilename)
-			: Name(MoveTemp(InName))
-			, NewFilename(MoveTemp(InFilename))
-		{}
-	};
-
 	/**
 	 * Adds a callback to directory watcher for the game binaries folder.
 	 */
@@ -163,7 +152,7 @@ private:
 	/**
 	 * Does the actual hot-reload, unloads old modules, loads new ones
 	 */
-	ECompilationResult::Type DoHotReloadInternal(const TArray<FRecompiledModule>& ChangedModules, const TArray<UPackage*>& Packages, const TArray<FName>& InDependentModules, FOutputDevice& HotReloadAr);
+	ECompilationResult::Type DoHotReloadInternal(const TMap<FString, FString>& ChangedModuleNames, const TArray<UPackage*>& Packages, const TArray<FName>& InDependentModules, FOutputDevice& HotReloadAr);
 
 	/**
 	 * Finds all references to old CDOs and replaces them with the new ones.
@@ -172,16 +161,6 @@ private:
 	 */
 	void ReplaceReferencesToReconstructedCDOs();
 
-	/**
-	 * Gets all currently loaded game module names and optionally, the file names for those modules
-	 */
-	void GetGameModules(TArray<FString>& OutGameModules, TArray<FString>* OutGameModuleFilePaths = nullptr);
-
-	/**
-	 * Gets packages to re-bind and dependent modules.
-	 */
-	void GetPackagesToRebindAndDependentModules(const TArray<FString>& InGameModuleNames, TArray<UPackage*>& OutPackagesToRebind, TArray<FName>& OutDependentModules);
-
 #if WITH_ENGINE
 	void RegisterForReinstancing(UClass* OldClass, UClass* NewClass);
 	void ReinstanceClasses();
@@ -215,10 +194,11 @@ private:
 	/**
 	 * Declares a function type that is executed after a module recompile has finished.
 	 *
-	 * The first argument signals whether compilation has finished.
-	 * The second argument shows whether compilation was successful or not.
+	 * ChangedModules: A map between the names of the modules that have changed and their filenames.
+	 * bRecompileFinished: Signals whether compilation has finished.
+	 * CompilationResult: Shows whether compilation was successful or not.
 	 */
-	typedef TFunction<void(const TArray<FRecompiledModule>& ChangedModules, bool bRecompileFinished, ECompilationResult::Type CompilationResult)> FRecompileModulesCallback;
+	typedef TFunction<void(const TMap<FString, FString>& ChangedModules, bool bRecompileFinished, ECompilationResult::Type CompilationResult)> FRecompileModulesCallback;
 
 	/**
 	 * Tries to recompile the specified modules in the background.  When recompiling finishes, the specified callback
@@ -305,7 +285,7 @@ private:
 	bool bIsHotReloadingFromEditor;
 	
 	/** New module DLLs */
-	TArray<FRecompiledModule> NewModules;
+	TMap<FString, FString> NewModules;
 
 	/** Moduels that have been recently recompiled from the editor **/
 	TSet<FString> ModulesRecentlyCompiledInTheEditor;
@@ -372,7 +352,7 @@ namespace HotReloadDefs
 
 IMPLEMENT_MODULE(FHotReloadModule, HotReload);
 
-namespace
+namespace UE4HotReload_Private
 {
 	/**
 	 * Gets editor runs directory.
@@ -434,11 +414,87 @@ namespace
 		}
 #endif // WITH_EDITOR
 	}
+
+	/**
+	 * Gets all currently loaded game module names and optionally, the file names for those modules
+	 */
+	TArray<FString> GetGameModuleNames(const FModuleManager& ModuleManager)
+	{
+		TArray<FString> Result;
+
+		// Ask the module manager for a list of currently-loaded gameplay modules
+		TArray<FModuleStatus> ModuleStatuses;
+		ModuleManager.QueryModules(ModuleStatuses);
+
+		for (FModuleStatus& ModuleStatus : ModuleStatuses)
+		{
+			// We only care about game modules that are currently loaded
+			if (ModuleStatus.bIsLoaded && ModuleStatus.bIsGameModule)
+			{
+				Result.Add(MoveTemp(ModuleStatus.Name));
+			}
+		}
+
+		return Result;
+	}
+
+	/**
+	 * Gets all currently loaded game module names and optionally, the file names for those modules
+	 */
+	TMap<FString, FString> GetGameModuleFilenames(const FModuleManager& ModuleManager)
+	{
+		TMap<FString, FString> Result;
+
+		// Ask the module manager for a list of currently-loaded gameplay modules
+		TArray< FModuleStatus > ModuleStatuses;
+		ModuleManager.QueryModules(ModuleStatuses);
+
+		for (FModuleStatus& ModuleStatus : ModuleStatuses)
+		{
+			// We only care about game modules that are currently loaded
+			if (ModuleStatus.bIsLoaded && ModuleStatus.bIsGameModule)
+			{
+				Result.Add(MoveTemp(ModuleStatus.Name), MoveTemp(ModuleStatus.FilePath));
+			}
+		}
+
+		return Result;
+	}
+
+	struct FPackagesAndDependentNames
+	{
+		TArray<UPackage*> Packages;
+		TArray<FName> DependentNames;
+	};
+
+	/**
+	 * Gets named packages and the names dependents.
+	 */
+	FPackagesAndDependentNames SplitByPackagesAndDependentNames(const TArray<FString>& ModuleNames)
+	{
+		FPackagesAndDependentNames Result;
+
+		for (const FString& ModuleName : ModuleNames)
+		{
+			FString PackagePath = TEXT("/Script/") + ModuleName;
+
+			if (UPackage* Package = FindPackage(nullptr, *PackagePath))
+			{
+				Result.Packages.Add(Package);
+			}
+			else
+			{
+				Result.DependentNames.Add(*ModuleName);
+			}
+		}
+
+		return Result;
+	}
 }
 
 void FHotReloadModule::StartupModule()
 {
-	CreateFileThatIndicatesEditorRunIfNeeded();
+	UE4HotReload_Private::CreateFileThatIndicatesEditorRunIfNeeded();
 
 	bIsHotReloadingFromEditor = false;
 
@@ -463,7 +519,7 @@ void FHotReloadModule::ShutdownModule()
 	FTicker::GetCoreTicker().RemoveTicker(TickerDelegateHandle);
 	ShutdownHotReloadWatcher();
 
-	DeleteFileThatIndicatesEditorRunIfNeeded();
+	UE4HotReload_Private::DeleteFileThatIndicatesEditorRunIfNeeded();
 }
 
 bool FHotReloadModule::Exec( UWorld* Inworld, const TCHAR* Cmd, FOutputDevice& Ar )
@@ -668,14 +724,14 @@ void FHotReloadModule::AddHotReloadFunctionRemap(Native NewFunctionPointer, Nati
 ECompilationResult::Type FHotReloadModule::DoHotReloadFromEditor(const bool bWaitForCompletion)
 {
 	// Get all game modules we want to compile
-	TArray<FString> GameModuleNames;
-	GetGameModules(GameModuleNames);
+	const FModuleManager& ModuleManager = FModuleManager::Get();
+	TArray<FString> GameModuleNames = UE4HotReload_Private::GetGameModuleNames(ModuleManager);
 
 	int32 NumPackagesToRebind = 0;
 	int32 NumDependentModules = 0;
 
 	ECompilationResult::Type Result = ECompilationResult::Unsupported;
-	
+
 	// Analytics
 	double Duration = 0.0;
 
@@ -683,13 +739,11 @@ ECompilationResult::Type FHotReloadModule::DoHotReloadFromEditor(const bool bWai
 	{
 		FScopedDurationTimer Timer(Duration);
 
-		TArray<UPackage*> PackagesToRebind;
-		TArray<FName> DependentModules;
-		GetPackagesToRebindAndDependentModules(GameModuleNames, PackagesToRebind, DependentModules);
+		UE4HotReload_Private::FPackagesAndDependentNames PackagesAndDependentNames = UE4HotReload_Private::SplitByPackagesAndDependentNames(GameModuleNames);
 
-		NumPackagesToRebind = PackagesToRebind.Num();
-		NumDependentModules = DependentModules.Num();
-		Result = RebindPackagesInternal(MoveTemp(PackagesToRebind), MoveTemp(DependentModules), bWaitForCompletion, *GLog);
+		NumPackagesToRebind = PackagesAndDependentNames.Packages.Num();
+		NumDependentModules = PackagesAndDependentNames.DependentNames.Num();
+		Result = RebindPackagesInternal(MoveTemp(PackagesAndDependentNames.Packages), MoveTemp(PackagesAndDependentNames.DependentNames), bWaitForCompletion, *GLog);
 	}
 
 	RecordAnalyticsEvent(TEXT("Editor"), Result, Duration, NumPackagesToRebind, NumDependentModules);
@@ -717,7 +771,7 @@ UObject* GetCachedCDODuplicate(UObject* CDO, FName Name)
 }
 #endif // WITH_HOT_RELOAD
 
-ECompilationResult::Type FHotReloadModule::DoHotReloadInternal(const TArray<FRecompiledModule>& ChangedModules, const TArray<UPackage*>& Packages, const TArray<FName>& InDependentModules, FOutputDevice& HotReloadAr)
+ECompilationResult::Type FHotReloadModule::DoHotReloadInternal(const TMap<FString, FString>& ChangedModules, const TArray<UPackage*>& Packages, const TArray<FName>& InDependentModules, FOutputDevice& HotReloadAr)
 {
 #if WITH_HOT_RELOAD
 
@@ -755,7 +809,7 @@ ECompilationResult::Type FHotReloadModule::DoHotReloadInternal(const TArray<FRec
 		FString PackageName = Package->GetName();
 		FString ShortPackageName = FPackageName::GetShortName(PackageName);
 
-		if (!ChangedModules.ContainsByPredicate([ShortPackageName](const FRecompiledModule& Module){ return Module.Name == ShortPackageName; }))
+		if (!ChangedModules.Contains(ShortPackageName))
 		{
 			continue;
 		}
@@ -764,8 +818,7 @@ ECompilationResult::Type FHotReloadModule::DoHotReloadInternal(const TArray<FRec
 
 		// Abandon the old module.  We can't unload it because various data structures may be living
 		// that have vtables pointing to code that would become invalidated.
-		const bool bAbandonModule = true;
-		ModuleManager.UnloadOrAbandonModuleWithCallback(ShortPackageFName, HotReloadAr, bAbandonModule);
+		ModuleManager.AbandonModuleWithCallback(ShortPackageFName);
 
 		// Load the newly-recompiled module up (it will actually have a different DLL file name at this point.)
 		bReloadSucceeded = ModuleManager.LoadModule(ShortPackageFName).IsValid();
@@ -781,8 +834,7 @@ ECompilationResult::Type FHotReloadModule::DoHotReloadInternal(const TArray<FRec
 	for (FName ModuleName : InDependentModules)
 	{
 		FString ModuleNameStr = ModuleName.ToString();
-		const TCHAR* ModuleNameStrPtr = *ModuleNameStr;
-		if (!ChangedModules.ContainsByPredicate([ModuleNameStrPtr](const FRecompiledModule& Module){ return Module.Name == ModuleNameStrPtr; }))
+		if (!ChangedModules.Contains(ModuleNameStr))
 		{
 			continue;
 		}
@@ -844,14 +896,8 @@ ECompilationResult::Type FHotReloadModule::DoHotReloadInternal(const TArray<FRec
 			check(Script->GetCppStructOps());
 		}
 		// Make sure new classes have the token stream assembled
-		for (FRawObjectIterator It; It; ++It)
-		{
-			UClass* Class = Cast<UClass>(static_cast<UObject*>(It->Object));
-			if (Class && !Class->HasAnyClassFlags(CLASS_TokenStreamAssembled))
-			{
-				Class->AssembleReferenceTokenStream();
-			}
-		}
+		UClass::AssembleReferenceTokenStreams();
+
 		HotReloadAr.Logf(ELogVerbosity::Display, TEXT("HotReload successful (%d functions remapped  %d scriptstructs remapped)"), NumFunctionsRemapped, ScriptStructs.Num());
 
 		HotReloadFunctionRemap.Empty();
@@ -907,27 +953,12 @@ void FHotReloadModule::ReplaceReferencesToReconstructedCDOs()
 		FQueuedThreadPool* Pool;
 	} ThreadPoolManager;
 
-	TArray<UObject*> OldCDOs;
-	ReconstructedCDOsMap.GetKeys(OldCDOs);
-
-	// Structure to store CDOs reference info.
-	struct FReferenceInfo
-	{
-		UObject* Referencer;
-		UObject* Referencee;
-		UProperty* ReferencingProperty;
-
-		FReferenceInfo(UObject* InReferencer, UObject* InReferencee, UProperty* InReferencingProperty)
-			: Referencer(InReferencer), Referencee(InReferencee), ReferencingProperty(InReferencingProperty)
-		{}
-	};
-
 	// Async task to enable multithreaded CDOs reference search.
 	class FFindRefTask : public FNonAbandonableTask
 	{
 	public:
-		FFindRefTask(const TArray<UObject*>* InReferencees, int32 ReserveElements = 0)
-			: Referencees(*InReferencees)
+		explicit FFindRefTask(const TMap<UObject*, UObject*>& InReconstructedCDOsMap, int32 ReserveElements)
+			: ReconstructedCDOsMap(InReconstructedCDOsMap)
 		{
 			ObjectsArray.Reserve(ReserveElements);
 		}
@@ -936,34 +967,45 @@ void FHotReloadModule::ReplaceReferencesToReconstructedCDOs()
 		{
 			for (UObject* Object : ObjectsArray)
 			{
-				FFindReferencersArchive FindRefsArchive(Object, Referencees);
-
-				TMap<UObject*, int32> ReferenceCounts;
-				TMultiMap<UObject*, UProperty*> ReferencingProperties;
+				class FReplaceCDOReferencesArchive : public FArchiveUObject
+				{
+				public:
+					FReplaceCDOReferencesArchive(UObject* InPotentialReferencer, const TMap<UObject*, UObject*>& InReconstructedCDOsMap)
+						: ReconstructedCDOsMap(InReconstructedCDOsMap)
+						, PotentialReferencer(InPotentialReferencer)
+					{
+						ArIsObjectReferenceCollector = true;
+						ArIgnoreOuterRef = true;
+					}
 
-				FindRefsArchive.GetReferenceCounts(ReferenceCounts, ReferencingProperties);
+					virtual FString GetArchiveName() const override
+					{
+						return TEXT("FReplaceCDOReferencesArchive");
+					}
 
-				for (auto& ReferencingProperty : ReferencingProperties)
-				{
-					static const UProperty* PropertyToSkip =
-#if WITH_ENGINE
-						UBlueprintGeneratedClass::StaticClass()->FindPropertyByName(GET_MEMBER_NAME_CHECKED(UBlueprintGeneratedClass, OverridenArchetypeForCDO));
-#else
-						nullptr;
-#endif
-					if (!ReferencingProperty.Value->IsA<UObjectProperty>() || ReferencingProperty.Value == PropertyToSkip)
+					FArchive& operator<<(UObject*& ObjRef)
 					{
-						continue;
+						UObject* Obj = ObjRef;
+
+						if (Obj && Obj != PotentialReferencer && ReconstructedCDOsMap.Contains(Obj))
+						{
+							UProperty* SerializedProp = GetSerializedProperty();
+							if (SerializedProp && SerializedProp->IsA<UObjectProperty>())
+							{
+								ObjRef = ReconstructedCDOsMap[Obj];
+							}
+						}
+					
+						return *this;
 					}
 
-					References.Emplace(Object, ReferencingProperty.Key, ReferencingProperty.Value);
-				}
-			}
-		}
+					const TMap<UObject*, UObject*>& ReconstructedCDOsMap;
+					UObject* PotentialReferencer;
+				};
 
-		TArray<UObject*>& GetObjectsArray()
-		{
-			return ObjectsArray;
+				FReplaceCDOReferencesArchive FindRefsArchive(Object, ReconstructedCDOsMap);
+				Object->Serialize(FindRefsArchive);
+			}
 		}
 
 		FORCEINLINE TStatId GetStatId() const
@@ -971,12 +1013,10 @@ void FHotReloadModule::ReplaceReferencesToReconstructedCDOs()
 			RETURN_QUICK_DECLARE_CYCLE_STAT(FFindRefTask, STATGROUP_ThreadPoolAsyncTasks);
 		}
 
-		const TArray<FReferenceInfo>& GetReferences() const { return References; }
+		TArray<UObject*> ObjectsArray;
 
 	private:
-		const TArray<UObject*>& Referencees;
-		TArray<UObject*> ObjectsArray;
-		TArray<FReferenceInfo> References;
+		const TMap<UObject*, UObject*>& ReconstructedCDOsMap;
 	};
 
 	const int32 NumberOfThreads = FPlatformMisc::NumberOfWorkerThreadsToSpawn();
@@ -989,7 +1029,7 @@ void FHotReloadModule::ReplaceReferencesToReconstructedCDOs()
 
 	for (int32 TaskId = 0; TaskId < NumberOfThreads; ++TaskId)
 	{
-		Tasks.Emplace(&OldCDOs, ObjectsPerTask);
+		Tasks.Emplace(ReconstructedCDOsMap, ObjectsPerTask);
 	}
 
 	// Distribute objects uniformly between tasks.
@@ -1003,31 +1043,20 @@ void FHotReloadModule::ReplaceReferencesToReconstructedCDOs()
 			continue;
 		}
 
-		Tasks[CurrentTaskId].GetTask().GetObjectsArray().Add(CurObject);
+		Tasks[CurrentTaskId].GetTask().ObjectsArray.Add(CurObject);
 		CurrentTaskId = (CurrentTaskId + 1) % NumberOfThreads;
 	}
 
 	// Run async tasks in worker threads.
-	for (int32 TaskId = 0; TaskId < NumberOfThreads; ++TaskId)
+	for (FAsyncTask<FFindRefTask>& Task : Tasks)
 	{
-		Tasks[TaskId].StartBackgroundTask(ThreadPoolManager.GetPool());
+		Task.StartBackgroundTask(ThreadPoolManager.GetPool());
 	}
 
-	// Wait until tasks are finished and replace found references
-	// in main thread.
-	for (int32 TaskId = 0; TaskId < NumberOfThreads; ++TaskId)
+	// Wait until tasks are finished
+	for (FAsyncTask<FFindRefTask>& AsyncTask : Tasks)
 	{
-		Tasks[TaskId].EnsureCompletion();
-
-		const TArray<FReferenceInfo>& References = Tasks[TaskId].GetTask().GetReferences();
-
-		for (const FReferenceInfo& Reference : References)
-		{
-			UObject* OldCDO = Reference.Referencee;
-			UObjectProperty* Prop = (UObjectProperty*)Reference.ReferencingProperty;
-
-			Prop->SetObjectPropertyValue((uint8*)Reference.Referencer + Prop->GetOffset_ForInternal(), ReconstructedCDOsMap[OldCDO]);
-		}
+		AsyncTask.EnsureCompletion();
 	}
 
 	ReconstructedCDOsMap.Empty();
@@ -1101,7 +1130,7 @@ ECompilationResult::Type FHotReloadModule::RebindPackagesInternal(TArray<UPackag
 		// Start compiling modules
 		const bool bCompileStarted = RecompileModulesAsync(
 			ModuleNames,
-			[this, InPackages/*=MoveTemp(InPackages)*/, DependentModules/*=MoveTemp(DependentModules)*/, &Ar](const TArray<FRecompiledModule>& ChangedModules, bool bRecompileFinished, ECompilationResult::Type CompilationResult)
+			[this, InPackages/*=MoveTemp(InPackages)*/, DependentModules/*=MoveTemp(DependentModules)*/, &Ar](const TMap<FString, FString>& ChangedModules, bool bRecompileFinished, ECompilationResult::Type CompilationResult)
 			{
 				if (ECompilationResult::Failed(CompilationResult) && bRecompileFinished)
 				{
@@ -1158,13 +1187,23 @@ void FHotReloadModule::RegisterForReinstancing(UClass* OldClass, UClass* NewClas
 	Pair.Key = OldClass;
 	Pair.Value = NewClass;
 
-	GetClassesToReinstance().Add(MoveTemp(Pair));
+	TArray<TPair<UClass*, UClass*> >& ClassesToReinstance = GetClassesToReinstance();
+	ClassesToReinstance.Add(MoveTemp(Pair));
 }
 
 void FHotReloadModule::ReinstanceClasses()
 {
+#if WITH_HOT_RELOAD
+	if (GIsHotReload)
+	{
+		UClass::AssembleReferenceTokenStreams();
+	}
+#endif // WITH_HOT_RELOAD
+
+	TArray<TPair<UClass*, UClass*> >& ClassesToReinstance = GetClassesToReinstance();
+
 	TMap<UClass*, UClass*> OldToNewClassesMap;
-	for (TPair<UClass*, UClass*>& Pair : GetClassesToReinstance())
+	for (const TPair<UClass*, UClass*>& Pair : ClassesToReinstance)
 	{
 		if (Pair.Value != nullptr)
 		{
@@ -1172,12 +1211,12 @@ void FHotReloadModule::ReinstanceClasses()
 		}
 	}
 
-	for (TPair<UClass*, UClass*>& Pair : GetClassesToReinstance())
+	for (const TPair<UClass*, UClass*>& Pair : ClassesToReinstance)
 	{
 		ReinstanceClass(Pair.Key, Pair.Value, OldToNewClassesMap);
 	}
 
-	GetClassesToReinstance().Empty();
+	ClassesToReinstance.Empty();
 }
 
 void FHotReloadModule::ReinstanceClass(UClass* OldClass, UClass* NewClass, const TMap<UClass*, UClass*>& OldToNewClassesMap)
@@ -1191,28 +1230,6 @@ void FHotReloadModule::ReinstanceClass(UClass* OldClass, UClass* NewClass, const
 }
 #endif
 
-void FHotReloadModule::GetGameModules(TArray<FString>& OutGameModules, TArray<FString>* OutGameModuleFilePaths )
-{
-	// Ask the module manager for a list of currently-loaded gameplay modules
-	TArray< FModuleStatus > ModuleStatuses;
-	FModuleManager::Get().QueryModules(ModuleStatuses);
-
-	for (TArray< FModuleStatus >::TConstIterator ModuleStatusIt(ModuleStatuses); ModuleStatusIt; ++ModuleStatusIt)
-	{
-		const FModuleStatus& ModuleStatus = *ModuleStatusIt;
-
-		// We only care about game modules that are currently loaded
-		if (ModuleStatus.bIsLoaded && ModuleStatus.bIsGameModule)
-		{
-			OutGameModules.Add(ModuleStatus.Name);
-			if( OutGameModuleFilePaths != nullptr )
-			{
-				OutGameModuleFilePaths->Add(ModuleStatus.FilePath);
-			}
-		}
-	}
-}
-
 void FHotReloadModule::OnHotReloadBinariesChanged(const TArray<struct FFileChangeData>& FileChanges)
 {
 	if (bIsHotReloadingFromEditor)
@@ -1221,50 +1238,63 @@ void FHotReloadModule::OnHotReloadBinariesChanged(const TArray<struct FFileChang
 		return;
 	}
 
-	TArray< FString > GameModuleNames;
-	TArray< FString > GameModuleFilePaths;
-	GetGameModules(GameModuleNames, &GameModuleFilePaths);
+	const FModuleManager& ModuleManager = FModuleManager::Get();
+	TMap<FString, FString> GameModuleFilenames = UE4HotReload_Private::GetGameModuleFilenames(ModuleManager);
 
-	if (GameModuleNames.Num() > 0)
+	if (GameModuleFilenames.Num() == 0)
 	{
-		// Check if any of the game DLLs has been added
-		for (auto& Change : FileChanges)
-		{
+		return;
+	}
+
+	// Check if any of the game DLLs has been added
+	for (auto& Change : FileChanges)
+	{
+		// Ignore changes that aren't introducing a new file.
+		//
+		// On the Mac the Add event is for a temporary linker(?) file that gets immediately renamed
+		// to a dylib. In the future we may want to support modified event for all platforms anyway once
+		// shadow copying works with hot-reload.
 #if PLATFORM_MAC
-			// On the Mac the Add event is for a temporary linker(?) file that gets immediately renamed
-			// to a dylib. In the future we may want to support modified event for all platforms anyway once
-			// shadow copying works with hot-reload.
-			if (Change.Action == FFileChangeData::FCA_Modified)
+		if (Change.Action != FFileChangeData::FCA_Modified)
 #else
-			if (Change.Action == FFileChangeData::FCA_Added)
+		if (Change.Action != FFileChangeData::FCA_Added)
 #endif
+		{
+			continue;
+		}
+
+		// Ignore files that aren't of module type
+		FString Filename = FPaths::GetCleanFilename(Change.Filename);
+		if (!Filename.EndsWith(FPlatformProcess::GetModuleExtension()))
+		{
+			continue;
+		}
+
+		for (const TPair<FString, FString>& NameFilename : GameModuleFilenames)
+		{
+			// Handle module files which have already been hot-reloaded.
+			FString BaseName = FPaths::GetBaseFilename(NameFilename.Value);
+			StripModuleSuffixFromFilename(BaseName, NameFilename.Key);
+
+			// Hot reload always adds a numbered suffix preceded by a hyphen, but otherwise the module name must match exactly!
+			if (!Filename.StartsWith(BaseName + TEXT("-")))
 			{
-				const FString Filename = FPaths::GetCleanFilename(Change.Filename);
-				if (Filename.EndsWith(FPlatformProcess::GetModuleExtension()))
-				{
-					for (int32 GameModuleIndex = 0; GameModuleIndex < GameModuleNames.Num(); ++GameModuleIndex)
-					{
-						const FString& GameModuleName = GameModuleNames[GameModuleIndex];
-						const FString& GameModuleFilePath = GameModuleFilePaths[GameModuleIndex];
+				continue;
+			}
 
-						// Handle module files which have already been hot-reloaded.
-						FString GameModuleFileNameWithoutExtension = FPaths::GetBaseFilename(GameModuleFilePath);
-						StripModuleSuffixFromFilename(GameModuleFileNameWithoutExtension, GameModuleName);
+			if (NewModules.Contains(NameFilename.Key))
+			{
+				continue;
+			}
 
-						// Hot reload always adds a numbered suffix preceded by a hyphen, but otherwise the module name must match exactly!
-						if (Filename.StartsWith(GameModuleFileNameWithoutExtension + TEXT("-")))
-						{
-							if (!NewModules.ContainsByPredicate([&](const FRecompiledModule& Module){ return Module.Name == GameModuleName; }) &&
-								!ModulesRecentlyCompiledInTheEditor.Contains(FPaths::ConvertRelativePathToFull(Change.Filename)))
-							{
-								// Add to queue. We do not hot-reload here as there may potentially be other modules being compiled.
-								NewModules.Emplace(GameModuleName, Change.Filename);
-								UE_LOG(LogHotReload, Log, TEXT("New module detected: %s"), *Filename);
-							}
-						}
-					}
-				}
+			if (ModulesRecentlyCompiledInTheEditor.Contains(FPaths::ConvertRelativePathToFull(Change.Filename)))
+			{
+				continue;
 			}
+
+			// Add to queue. We do not hot-reload here as there may potentially be other modules being compiled.
+			NewModules.Emplace(NameFilename.Key, Change.Filename);
+			UE_LOG(LogHotReload, Log, TEXT("New module detected: %s"), *Filename);
 		}
 	}
 }
@@ -1342,37 +1372,38 @@ bool FHotReloadModule::Tick(float DeltaTime)
 	return true;
 }
 
-void FHotReloadModule::GetPackagesToRebindAndDependentModules(const TArray<FString>& InGameModuleNames, TArray<UPackage*>& OutPackagesToRebind, TArray<FName>& OutDependentModules)
-{
-	for (auto& GameModuleName : InGameModuleNames)
-	{
-		FString PackagePath(FString(TEXT("/Script/")) + GameModuleName);
-		UPackage* Package = FindPackage(NULL, *PackagePath);
-		if (Package != NULL)
-		{
-			OutPackagesToRebind.Add(Package);
-		}
-		else
-		{
-			OutDependentModules.Add(*GameModuleName);
-		}
-	}
-}
-
 void FHotReloadModule::DoHotReloadFromIDE()
 {
-	TArray<FString> GameModuleNames;
-	TArray<UPackage*> PackagesToRebind;
-	TArray<FName> DependentModules;
-	double Duration = 0.0;
+	const FModuleManager& ModuleManager = FModuleManager::Get();
+	IFileManager& FileManager = IFileManager::Get();
+
+	int32 NumPackagesToRebind = 0;
+	int32 NumDependentModules = 0;
+
 	ECompilationResult::Type Result = ECompilationResult::Unsupported;
 
-	GetGameModules(GameModuleNames);
+	double Duration = 0.0;
+
+	TArray<FString> GameModuleNames = UE4HotReload_Private::GetGameModuleNames(ModuleManager);
 
 	if (GameModuleNames.Num() > 0)
 	{
 		FScopedDurationTimer Timer(Duration);
 
+		// Remove any modules whose files have disappeared - this can happen if a compile event has
+		// failed and deleted a DLL that was there previously.
+		for (auto It = NewModules.CreateIterator(); It; ++It)
+		{
+			if (!FileManager.FileExists(*It->Value))
+			{
+				It.RemoveCurrent();
+			}
+		}
+		if (NewModules.Num() == 0)
+		{
+			return;
+		}
+
 		UE_LOG(LogHotReload, Log, TEXT("Starting Hot-Reload from IDE"));
 
 		HotReloadStartTime = FPlatformTime::Seconds();
@@ -1381,25 +1412,27 @@ void FHotReloadModule::DoHotReloadFromIDE()
 		SlowTask.MakeDialog();
 
 		// Update compile data before we start compiling
-		for (auto& NewModule : NewModules)
+		for (const TPair<FString, FString>& NewModule : NewModules)
 		{
 			// Move on 10% / num items
 			SlowTask.EnterProgressFrame(10.f/NewModules.Num());
 
-			UpdateModuleCompileData(*NewModule.Name);
-			OnModuleCompileSucceeded(*NewModule.Name, NewModule.NewFilename);
+			FName ModuleName = *NewModule.Key;
+
+			UpdateModuleCompileData(ModuleName);
+			OnModuleCompileSucceeded(ModuleName, NewModule.Value);
 		}
 
 		SlowTask.EnterProgressFrame(10);
-		GetPackagesToRebindAndDependentModules(GameModuleNames, PackagesToRebind, DependentModules);
-
+		UE4HotReload_Private::FPackagesAndDependentNames PackagesAndDependentNames = UE4HotReload_Private::SplitByPackagesAndDependentNames(GameModuleNames);
 		SlowTask.EnterProgressFrame(80);
-		check(PackagesToRebind.Num() || DependentModules.Num());
 
-		Result = DoHotReloadInternal(NewModules, PackagesToRebind, DependentModules, *GLog);
+		NumPackagesToRebind = PackagesAndDependentNames.Packages.Num();
+		NumDependentModules = PackagesAndDependentNames.DependentNames.Num();
+		Result = DoHotReloadInternal(NewModules, PackagesAndDependentNames.Packages, PackagesAndDependentNames.DependentNames, *GLog);
 	}
 
-	RecordAnalyticsEvent(TEXT("IDE"), Result, Duration, PackagesToRebind.Num(), DependentModules.Num());
+	RecordAnalyticsEvent(TEXT("IDE"), Result, Duration, NumPackagesToRebind, NumDependentModules);
 }
 
 void FHotReloadModule::RecordAnalyticsEvent(const TCHAR* ReloadFrom, ECompilationResult::Type Result, double Duration, int32 PackageCount, int32 DependentModulesCount)
@@ -1619,7 +1652,7 @@ bool FHotReloadModule::StartCompilingModuleDLLs(const FString& GameName, const T
 		
 		if (RecompileModulesCallback)
 		{
-			RecompileModulesCallback( TArray<FRecompiledModule>(), false, ECompilationResult::OtherCompilationError );
+			RecompileModulesCallback( TMap<FString, FString>(), false, ECompilationResult::OtherCompilationError );
 			RecompileModulesCallback = nullptr;
 		}
 	}
@@ -1753,7 +1786,7 @@ void FHotReloadModule::CheckForFinishedModuleDLLCompile(const bool bWaitForCompl
 	// in case we recompiled the modules to a new unique file name.  This is needed so that when the module
 	// is reloaded after the recompile, we load the new DLL file name, not the old one.
 	// Note that we don't want to do anything in case the build was canceled or source code has not changed.
-	TArray<FRecompiledModule> ChangedModules;
+	TMap<FString, FString> ChangedModules;
 	if(CompilationResult == ECompilationResult::Succeeded)
 	{
 		ChangedModules.Reserve(ModulesThatWereBeingRecompiled.Num());
diff --git a/Engine/Source/Developer/MeshSimplifier/Private/BinaryHeap.h b/Engine/Source/Developer/MeshSimplifier/Private/BinaryHeap.h
--- a/Engine/Source/Developer/MeshSimplifier/Private/BinaryHeap.h
+++ b/Engine/Source/Developer/MeshSimplifier/Private/BinaryHeap.h
@@ -157,6 +157,8 @@ void FBinaryHeap< KeyType >::ResizeIndexes( uint32 NewIndexSize )
 
 	if( IndexSize != 0 )
 	{
+		check( NewIndexSize >= IndexSize );
+
 		for( uint32 i = 0; i < IndexSize; i++ )
 		{
 			NewKeys[i] = Keys[i];
diff --git a/Engine/Source/Developer/MeshSimplifier/Private/MeshSimplify.h b/Engine/Source/Developer/MeshSimplifier/Private/MeshSimplify.h
--- a/Engine/Source/Developer/MeshSimplifier/Private/MeshSimplify.h
+++ b/Engine/Source/Developer/MeshSimplifier/Private/MeshSimplify.h
@@ -208,48 +208,52 @@ void TMeshSimplifier<T, NumAttributes>::SetBoundaryLocked()
 
 		LockVertFlags( SIMP_MARK1 );
 
-		TSimpVert<T>** adjVerts = (TSimpVert<T>**)FMemory_Alloca( 2 * v0->NumAdjTrisGroup() * sizeof( TSimpVert<T>* ) );
-		int numAdjVerts;
-		v0->FindAdjacentVertsGroup( adjVerts, numAdjVerts );
+		[this, v0]()
+		{
+			TSimpVert<T>** adjVerts = (TSimpVert<T>**)FMemory_Alloca( 2 * v0->NumAdjTrisGroup() * sizeof( TSimpVert<T>* ) );
+			check(adjVerts);
+			int numAdjVerts;
+			v0->FindAdjacentVertsGroup( adjVerts, numAdjVerts );
 
-		UnlockVertFlags( SIMP_MARK1 );
+			UnlockVertFlags( SIMP_MARK1 );
 
-		for( int a = 0; a < numAdjVerts; a++ )
-		{
-			TSimpVert<T>* v1 = adjVerts[a];
-			if( v0 < v1 )
+			for( int a = 0; a < numAdjVerts; a++ )
 			{
-				LockTriFlags( SIMP_MARK1 );
+				TSimpVert<T>* v1 = adjVerts[a];
+				if( v0 < v1 )
+				{
+					LockTriFlags( SIMP_MARK1 );
 
-				// set if this edge is boundary
-				// find faces that share v0 and v1
-				v0->EnableAdjTriFlagsGroup( SIMP_MARK1 );
-				v1->DisableAdjTriFlagsGroup( SIMP_MARK1 );
+					// set if this edge is boundary
+					// find faces that share v0 and v1
+					v0->EnableAdjTriFlagsGroup( SIMP_MARK1 );
+					v1->DisableAdjTriFlagsGroup( SIMP_MARK1 );
 
-				int faceCount = 0;
-				TSimpVert<T>* vert = v0;
-				do
-				{
-					for( TriIterator j = vert->adjTris.Begin(); j != vert->adjTris.End(); ++j )
+					int faceCount = 0;
+					TSimpVert<T>* vert = v0;
+					do
 					{
-						TSimpTri<T>* tri = *j;
-						faceCount += tri->TestFlags( SIMP_MARK1 ) ? 0 : 1;
-					}
-					vert = vert->next;
-				} while( vert != v0 );
+						for( TriIterator j = vert->adjTris.Begin(); j != vert->adjTris.End(); ++j )
+						{
+							TSimpTri<T>* tri = *j;
+							faceCount += tri->TestFlags( SIMP_MARK1 ) ? 0 : 1;
+						}
+						vert = vert->next;
+					} while( vert != v0 );
 
-				v0->DisableAdjTriFlagsGroup( SIMP_MARK1 );
+					v0->DisableAdjTriFlagsGroup( SIMP_MARK1 );
 
-				if( faceCount == 1 )
-				{
-					// only one face on this edge
-					v0->EnableFlagsGroup( SIMP_LOCKED );
-					v1->EnableFlagsGroup( SIMP_LOCKED );
-				}
+					if( faceCount == 1 )
+					{
+						// only one face on this edge
+						v0->EnableFlagsGroup( SIMP_LOCKED );
+						v1->EnableFlagsGroup( SIMP_LOCKED );
+					}
 
-				UnlockTriFlags( SIMP_MARK1 );
+					UnlockTriFlags( SIMP_MARK1 );
+				}
 			}
-		}
+		}();
 	}
 }
 
@@ -331,6 +335,7 @@ void TMeshSimplifier<T, NumAttributes>::InitVert( TSimpVert<T>* v )
 	LockVertFlags( SIMP_MARK1 );
 
 	TSimpVert<T>** adjVerts = (TSimpVert<T>**)FMemory_Alloca( 2 * v->adjTris.Num() * sizeof( TSimpVert<T>* ) );
+	check(adjVerts);
 	int numAdjVerts;
 	v->FindAdjacentVerts( adjVerts, numAdjVerts );
 
@@ -518,6 +523,7 @@ FQuadric TMeshSimplifier<T, NumAttributes>::GetEdgeQuadric( TSimpVert<T>* v )
 	LockVertFlags( SIMP_MARK1 );
 
 	TSimpVert<T>** adjVerts = (TSimpVert<T>**)FMemory_Alloca( 2 * v->adjTris.Num() * sizeof( TSimpVert<T>* ) );
+	check(adjVerts);
 	int numAdjVerts;
 	v->FindAdjacentVerts( adjVerts, numAdjVerts );
 
@@ -755,6 +761,7 @@ void TMeshSimplifier<T, NumAttributes>::GatherUpdates( TSimpVert<T>* v )
 
 	// update the costs of all edges connected to any face adjacent to v
 	TSimpVert<T>** adjVerts = (TSimpVert<T>**)FMemory_Alloca( 2 * v->adjTris.Num() * sizeof( TSimpVert<T>* ) );
+	check(adjVerts);
 	int numAdjVerts;
 	v->FindAdjacentVerts( adjVerts, numAdjVerts );
 
@@ -901,6 +908,7 @@ void TMeshSimplifier<T, NumAttributes>::UpdateEdgeCosts( TSimpVert<T>* vert ) {
 	}
 }
 
+CA_SUPPRESS(6262); // warning C6262: Function uses '121180' bytes of stack:  exceeds /analyze:stacksize '81940'.  Consider moving some data to heap.
 template< typename T, uint32 NumAttributes >
 float TMeshSimplifier<T, NumAttributes>::ComputeNewVerts( TSimpEdge<T>* edge, T* newVerts )
 {
@@ -1399,9 +1407,7 @@ void TMeshSimplifier<T, NumAttributes>::SimplifyMesh( float maxError, int minTri
 			uint32	vertListNum = 0;
 			TSimpVert<T>*	vertList[256];
 			
-			PRAGMA_DISABLE_SHADOW_VARIABLE_WARNINGS
-			TSimpVert<T>* v = top->v1;
-			PRAGMA_ENABLE_SHADOW_VARIABLE_WARNINGS
+			v = top->v1;
 			do {
 				vertList[ vertListNum++ ] = v;
 				v = v->next;
@@ -1410,9 +1416,7 @@ void TMeshSimplifier<T, NumAttributes>::SimplifyMesh( float maxError, int minTri
 			check( vertListNum <= 256 );
 			
 			for( uint32 i = 0; i < vertListNum; i++ ) {
-				PRAGMA_DISABLE_SHADOW_VARIABLE_WARNINGS
-				TSimpVert<T>* v = vertList[i];
-				PRAGMA_ENABLE_SHADOW_VARIABLE_WARNINGS
+				v = vertList[i];
 
 				if( v->TestFlags( SIMP_REMOVED ) ) {
 					// ungroup
@@ -1427,9 +1431,6 @@ void TMeshSimplifier<T, NumAttributes>::SimplifyMesh( float maxError, int minTri
 		{
 			// spread locked flag to vert group
 			uint32 flags = 0;
-			PRAGMA_DISABLE_SHADOW_VARIABLE_WARNINGS
-			TSimpVert<T>* v;
-			PRAGMA_ENABLE_SHADOW_VARIABLE_WARNINGS
 
 			v = top->v1;
 			do {
@@ -1505,9 +1506,7 @@ void TMeshSimplifier<T, NumAttributes>::SimplifyMesh( float maxError, int minTri
 
 		// remove dead edges
 		for( int i = 0; i < edges.Num(); i++ ) {
-			PRAGMA_DISABLE_SHADOW_VARIABLE_WARNINGS
-			TSimpEdge<T>* edge = &edges[i];
-			PRAGMA_ENABLE_SHADOW_VARIABLE_WARNINGS
+			edge = &edges[i];
 
 			if( edge->TestFlags( SIMP_REMOVED ) )
 				continue;
@@ -1523,9 +1522,7 @@ void TMeshSimplifier<T, NumAttributes>::SimplifyMesh( float maxError, int minTri
 			FHashTable HashTable( 4096, edges.Num() );
 
 			for( int i = 0; i < edges.Num(); i++ ) {
-				PRAGMA_DISABLE_SHADOW_VARIABLE_WARNINGS
-				TSimpEdge<T>* edge = &edges[i];
-				PRAGMA_ENABLE_SHADOW_VARIABLE_WARNINGS
+				edge = &edges[i];
 
 				if( edge->TestFlags( SIMP_REMOVED ) )
 					continue;
@@ -1580,9 +1577,7 @@ void TMeshSimplifier<T, NumAttributes>::SimplifyMesh( float maxError, int minTri
 
 		// update edges
 		for( uint32 i = 0; i < updateEdgesNum; i++ ) {
-			PRAGMA_DISABLE_SHADOW_VARIABLE_WARNINGS
-			TSimpEdge<T>* edge = updateEdges[i];
-			PRAGMA_ENABLE_SHADOW_VARIABLE_WARNINGS
+			edge = updateEdges[i];
 
 			edge->DisableFlags( SIMP_UPDATE );
 
@@ -1591,9 +1586,7 @@ void TMeshSimplifier<T, NumAttributes>::SimplifyMesh( float maxError, int minTri
 
 			float cost = ComputeEdgeCollapseCost( edge );
 
-			PRAGMA_DISABLE_SHADOW_VARIABLE_WARNINGS
-			TSimpEdge<T>* e = edge;
-			PRAGMA_ENABLE_SHADOW_VARIABLE_WARNINGS
+			e = edge;
 			do {
 				uint32 EdgeIndex = GetEdgeIndex(e);
 				if( edgeHeap.IsPresent( EdgeIndex ) )
diff --git a/Engine/Source/Developer/MeshSimplifier/Private/MeshSimplifyElements.h b/Engine/Source/Developer/MeshSimplifier/Private/MeshSimplifyElements.h
--- a/Engine/Source/Developer/MeshSimplifier/Private/MeshSimplifyElements.h
+++ b/Engine/Source/Developer/MeshSimplifier/Private/MeshSimplifyElements.h
@@ -285,6 +285,7 @@ FORCEINLINE uint32 TSimpVert<T>::NumAdjTrisGroup()
 	do {
 		numAdjVerts += v->adjTris.Num();
 		v = v->next;
+		check(v);
 	} while( v != this );
 
 	return numAdjVerts;
diff --git a/Engine/Source/Developer/Profiler/Private/Widgets/SDataGraph.cpp b/Engine/Source/Developer/Profiler/Private/Widgets/SDataGraph.cpp
--- a/Engine/Source/Developer/Profiler/Private/Widgets/SDataGraph.cpp
+++ b/Engine/Source/Developer/Profiler/Private/Widgets/SDataGraph.cpp
@@ -520,7 +520,7 @@ int32 SDataGraph::OnPaint( const FPaintArgs& Args, const FGeometry& AllottedGeom
 					continue;
 				}
 
-				// Bottom - Frame numbers, starting from 0.
+				// Top - Frame numbers, starting from 0.
 				const FString AccumulatedFrameNumberStr = FString::Printf(TEXT("%i"), FrameIndex);
 				FSlateDrawElement::MakeText
 				(
@@ -534,7 +534,7 @@ int32 SDataGraph::OnPaint( const FPaintArgs& Args, const FGeometry& AllottedGeom
 					FLinearColor::White
 				);
 
-				// Top - Time, normalized to the beginning of the capture process.
+				// Bottom - Time, normalized to the beginning of the capture process.
 				const FString ElapseTimeStr = FString::Printf(TEXT("%.1fs"), ElapsedFrameTimeMS * 0.001f );
 				FSlateDrawElement::MakeText
 				(
@@ -585,7 +585,7 @@ int32 SDataGraph::OnPaint( const FPaintArgs& Args, const FGeometry& AllottedGeom
 					continue;
 				}
 
-				// Bottom - Time, normalized to the beginning of the capture process.
+				// Top - Time, normalized to the beginning of the capture process.
 				const FString ElapsedTimeStr = FString::Printf(TEXT("%.1fs"), ElapsedFrameTimeMS * 0.001f);
 				FSlateDrawElement::MakeText
 				(
@@ -599,7 +599,7 @@ int32 SDataGraph::OnPaint( const FPaintArgs& Args, const FGeometry& AllottedGeom
 					FLinearColor::White
 				);
 
-				// Top - Frame numbers, starting from 0, for single session current frame counter.
+				// Bottom - Frame numbers, starting from 0, for single session current frame counter.
 				const FString AccumulatedFrameNumberStr = FString::Printf(TEXT("%i"), AccumulatedFrameCounter);
 				FSlateDrawElement::MakeText
 				(
diff --git a/Engine/Source/Developer/Profiler/Private/Widgets/SHistogram.cpp b/Engine/Source/Developer/Profiler/Private/Widgets/SHistogram.cpp
--- a/Engine/Source/Developer/Profiler/Private/Widgets/SHistogram.cpp
+++ b/Engine/Source/Developer/Profiler/Private/Widgets/SHistogram.cpp
@@ -21,9 +21,6 @@ void SHistogram::Construct( const FArguments& InArgs )
 
 int32 SHistogram::OnPaint( const FPaintArgs& Args, const FGeometry& AllottedGeometry, const FSlateRect& MyClippingRect, FSlateWindowElementList& OutDrawElements, int32 LayerId, const FWidgetStyle& InWidgetStyle, bool bParentEnabled ) const
 {
-	static double TotalTime = 0.0f;
-	static uint32 NumCalls = 0;
-	const double StartTime = FPlatformTime::Seconds();
 	const TSharedRef< FSlateFontMeasure > FontMeasureService = FSlateApplication::Get().GetRenderer()->GetFontMeasureService();
 
 	// Rendering info.
@@ -168,6 +165,11 @@ int32 SHistogram::OnPaint( const FPaintArgs& Args, const FGeometry& AllottedGeom
 	return SCompoundWidget::OnPaint( Args, AllottedGeometry, MyClippingRect, OutDrawElements, LayerId, InWidgetStyle, bParentEnabled && IsEnabled() );
 }
 
+void SHistogram::SetFPSAnalyzer(const TSharedPtr<FFPSAnalyzer>& InAnalyzer)
+{
+	Description.HistogramDataSource = InAnalyzer;
+}
+
 #undef LOCTEXT_NAMESPACE
 
 
diff --git a/Engine/Source/Developer/Profiler/Private/Widgets/SHistogram.h b/Engine/Source/Developer/Profiler/Private/Widgets/SHistogram.h
--- a/Engine/Source/Developer/Profiler/Private/Widgets/SHistogram.h
+++ b/Engine/Source/Developer/Profiler/Private/Widgets/SHistogram.h
@@ -111,6 +111,9 @@ public:
 		return FVector2D(16.0f,16.0f);
 	}
 
+	/** Sets a new Analyzer for the Description panel */
+	void SetFPSAnalyzer(const TSharedPtr<FFPSAnalyzer>& InAnalyzer);
+
 protected:
 	FHistogramDescription Description;
 };
\ No newline at end of file
diff --git a/Engine/Source/Developer/Profiler/Private/Widgets/SProfilerFPSChartPanel.cpp b/Engine/Source/Developer/Profiler/Private/Widgets/SProfilerFPSChartPanel.cpp
--- a/Engine/Source/Developer/Profiler/Private/Widgets/SProfilerFPSChartPanel.cpp
+++ b/Engine/Source/Developer/Profiler/Private/Widgets/SProfilerFPSChartPanel.cpp
@@ -101,6 +101,11 @@ public:
 		];
 	}
 	END_SLATE_FUNCTION_BUILD_OPTIMIZATION
+		
+	void SetFPSAnalyzer(const TSharedPtr<FFPSAnalyzer>& InAnalyzer)
+	{
+		FPSAnalyzer = InAnalyzer;
+	}
 
 protected:
 	FText HandleSampleCount() const
@@ -168,6 +173,16 @@ protected:
 	TSharedPtr<FFPSAnalyzer> FPSAnalyzer;
 };
 
+SProfilerFPSChartPanel::~SProfilerFPSChartPanel()
+{
+	// Remove ourselves from the profiler manager.
+	auto ProfilerManager = FProfilerManager::Get();
+	if (ProfilerManager.IsValid())
+	{
+		ProfilerManager->OnViewModeChanged().RemoveAll(this);
+	}
+}
+
 BEGIN_SLATE_FUNCTION_BUILD_OPTIMIZATION
 void SProfilerFPSChartPanel::Construct( const FArguments& InArgs )
 {
@@ -182,19 +197,37 @@ void SProfilerFPSChartPanel::Construct( const FArguments& InArgs )
 			+SHorizontalBox::Slot()
 			.FillWidth(1.0f)
 			[
-				SNew(SHistogram)
+				SAssignNew(Histogram, SHistogram)
 				.Description(FHistogramDescription(InArgs._FPSAnalyzer.ToSharedRef(), 5, 0, 90, true))
 			]
 
 			+SHorizontalBox::Slot()
 			.AutoWidth()
 			[
-				SNew( SProfilerFPSStatisticsPanel )
+				SAssignNew(StatisticsPanel, SProfilerFPSStatisticsPanel)
 				.FPSAnalyzer( InArgs._FPSAnalyzer )
 			]
 		]	
 	];
+
+	auto ProfilerManager = FProfilerManager::Get();
+	if (ProfilerManager.IsValid())
+	{
+		ProfilerManager->OnViewModeChanged().AddSP(this, &SProfilerFPSChartPanel::ProfilerManager_OnViewModeChanged);
+	}
 }
 END_SLATE_FUNCTION_BUILD_OPTIMIZATION
 
+void SProfilerFPSChartPanel::ProfilerManager_OnViewModeChanged(EProfilerViewMode::Type NewViewMode)
+{
+	auto ProfilerManager = FProfilerManager::Get();
+	if (ProfilerManager.IsValid())
+	{
+		const FProfilerSessionPtr ProfilerSession = ProfilerManager->GetProfilerSession();
+
+		Histogram->SetFPSAnalyzer(ProfilerSession->FPSAnalyzer);
+		StatisticsPanel->SetFPSAnalyzer(ProfilerSession->FPSAnalyzer);
+	}
+}
+
 #undef LOCTEXT_NAMESPACE
diff --git a/Engine/Source/Developer/Profiler/Private/Widgets/SProfilerFPSChartPanel.h b/Engine/Source/Developer/Profiler/Private/Widgets/SProfilerFPSChartPanel.h
--- a/Engine/Source/Developer/Profiler/Private/Widgets/SProfilerFPSChartPanel.h
+++ b/Engine/Source/Developer/Profiler/Private/Widgets/SProfilerFPSChartPanel.h
@@ -12,10 +12,23 @@ public:
 		SLATE_ARGUMENT( TSharedPtr<FFPSAnalyzer>, FPSAnalyzer )
 	SLATE_END_ARGS()
 
+	/** Virtual destructor. */
+	virtual ~SProfilerFPSChartPanel();
+
 	/**
 	 * Construct this widget
 	 *
 	 * @param InArgs - The declaration data for this widget
 	 */
 	void Construct( const FArguments& InArgs );
+
+protected:
+	/** Called when the status of specified tracked stat has changed. */
+	void ProfilerManager_OnViewModeChanged(EProfilerViewMode::Type NewViewMode);
+
+	/** The descriptions panel of the chart */
+	TSharedPtr<class SHistogram> Histogram;
+
+	/** The descriptions panel of the chart */
+	TSharedPtr<class SProfilerFPSStatisticsPanel> StatisticsPanel;
 };
\ No newline at end of file
diff --git a/Engine/Source/Developer/Profiler/Private/Widgets/SProfilerThreadView.cpp b/Engine/Source/Developer/Profiler/Private/Widgets/SProfilerThreadView.cpp
--- a/Engine/Source/Developer/Profiler/Private/Widgets/SProfilerThreadView.cpp
+++ b/Engine/Source/Developer/Profiler/Private/Widgets/SProfilerThreadView.cpp
@@ -194,7 +194,7 @@ void SProfilerThreadView::DrawFramesBackgroundAndTimelines() const
 				&SolidWhiteBrush,
 				PaintState->AbsoluteClippingRect,
 				PaintState->DrawEffects,
-				ThreadNode.FrameIndex % 2 ? FColorList::White.WithAlpha( 64 ) : FColorList::White.WithAlpha( 128 )
+				(ThreadNode.FrameIndex % 2) ? FColorList::White.WithAlpha( 64 ) : FColorList::White.WithAlpha( 128 )
 			);
 
 			// Check if this frame time marker is inside the visible area.
@@ -283,7 +283,7 @@ void SProfilerThreadView::DrawUIStackNodes() const
 				}
 
 				FColor NodeColor = UIStackNode->bIsCombined ? ThreadColors[UIStackNode->ThreadIndex].WithAlpha( 64 ) : ThreadColors[UIStackNode->ThreadIndex].WithAlpha( 192 );
-				NodeColor.G += NodeIndex % 2 ? 0 : 64;
+				NodeColor.G += (NodeIndex % 2) ? 0 : 64;
 
 				// Draw a cycle counter for this profiler UI stack node.
 				FSlateDrawElement::MakeBox
diff --git a/Engine/Source/Developer/SimplygonSwarm/Private/SimplygonRESTClient.cpp b/Engine/Source/Developer/SimplygonSwarm/Private/SimplygonRESTClient.cpp
--- a/Engine/Source/Developer/SimplygonSwarm/Private/SimplygonRESTClient.cpp
+++ b/Engine/Source/Developer/SimplygonSwarm/Private/SimplygonRESTClient.cpp
@@ -70,7 +70,7 @@ void FSimplygonSwarmTask::SetState(SimplygonRESTState InState)
 
 		//do not update the sate if either of these set is already set
 		if (InState != State
-			&& (State != SRS_ASSETDOWNLOADED || State != SRS_FAILED))
+			&& (State != SRS_ASSETDOWNLOADED && State != SRS_FAILED))
 		{
 			State = InState;
 		}
diff --git a/Engine/Source/Developer/Windows/WindowsTargetPlatform/Private/LocalPcTargetDevice.h b/Engine/Source/Developer/Windows/WindowsTargetPlatform/Private/LocalPcTargetDevice.h
--- a/Engine/Source/Developer/Windows/WindowsTargetPlatform/Private/LocalPcTargetDevice.h
+++ b/Engine/Source/Developer/Windows/WindowsTargetPlatform/Private/LocalPcTargetDevice.h
@@ -136,10 +136,10 @@ public:
 									if (::GetTokenInformation(TokenHandle, TokenUser, UserToken, UserTokenSize, &UserTokenSize))
 									{
 										WCHAR DomainName[256];
-										::DWORD DomainNameLength;
+										::DWORD DomainNameLength = 256;
 
 										WCHAR UserName[256];
-										::DWORD UserNameLength;
+										::DWORD UserNameLength = 256;
 
 										SID_NAME_USE SidType;
 
@@ -272,6 +272,7 @@ public:
 			return false;
 		}
 
+		CA_SUPPRESS(28159);
 		return (::ExitWindowsEx(EWX_POWEROFF | EWX_FORCE, SHTDN_REASON_MINOR_MAINTENANCE | SHTDN_REASON_FLAG_PLANNED) != 0);
 	}
 
@@ -287,6 +288,7 @@ public:
 			return false;
 		}
 
+		CA_SUPPRESS(28159);
 		return (::ExitWindowsEx(EWX_REBOOT | EWX_FORCE, SHTDN_REASON_MINOR_MAINTENANCE | SHTDN_REASON_FLAG_PLANNED) != 0);
 	}
 
diff --git a/Engine/Source/Editor/AnimGraph/Private/AnimGraphNode_SequencePlayer.cpp b/Engine/Source/Editor/AnimGraph/Private/AnimGraphNode_SequencePlayer.cpp
--- a/Engine/Source/Editor/AnimGraph/Private/AnimGraphNode_SequencePlayer.cpp
+++ b/Engine/Source/Editor/AnimGraph/Private/AnimGraphNode_SequencePlayer.cpp
@@ -70,7 +70,12 @@ FLinearColor UAnimGraphNode_SequencePlayer::GetNodeTitleColor() const
 
 FText UAnimGraphNode_SequencePlayer::GetTooltipText() const
 {
-	const bool bAdditive = ((Node.Sequence != NULL) && Node.Sequence->IsValidAdditive());
+	if (!Node.Sequence)
+	{
+		return GetTitleGivenAssetInfo(FText::FromString(TEXT("None")), false);
+	}
+
+	const bool bAdditive = Node.Sequence->IsValidAdditive();
 	return GetTitleGivenAssetInfo(FText::FromString(Node.Sequence->GetPathName()), bAdditive);
 }
 
diff --git a/Engine/Source/Editor/BlueprintGraph/Private/CallFunctionHandler.cpp b/Engine/Source/Editor/BlueprintGraph/Private/CallFunctionHandler.cpp
--- a/Engine/Source/Editor/BlueprintGraph/Private/CallFunctionHandler.cpp
+++ b/Engine/Source/Editor/BlueprintGraph/Private/CallFunctionHandler.cpp
@@ -210,16 +210,19 @@ void FKCHandler_CallFunction::CreateFunctionCallStatement(FKismetFunctionContext
 										}
 
 										int32 StructSize = Struct->GetStructureSize();
-										uint8* StructData = (uint8*)FMemory_Alloca(StructSize);
-										StructProperty->InitializeValue(StructData);
-
-										// Import the literal text to a dummy struct to verify it's well-formed
-										FImportTextErrorContext ErrorPipe(CompilerContext.MessageLog, Node);
-										StructProperty->ImportText(*((*Term)->Name), StructData, 0, NULL, &ErrorPipe);
-										if(ErrorPipe.NumErrors > 0)
+										[this, StructSize, StructProperty, Node, Term, &bMatchedAllParams]()
 										{
-											bMatchedAllParams = false;
-										}
+											uint8* StructData = (uint8*)FMemory_Alloca(StructSize);
+											StructProperty->InitializeValue(StructData);
+
+											// Import the literal text to a dummy struct to verify it's well-formed
+											FImportTextErrorContext ErrorPipe(CompilerContext.MessageLog, Node);
+											StructProperty->ImportText(*((*Term)->Name), StructData, 0, NULL, &ErrorPipe);
+											if(ErrorPipe.NumErrors > 0)
+											{
+												bMatchedAllParams = false;
+											}
+										}();
 									}
 									
 								}
diff --git a/Engine/Source/Editor/IntroTutorials/Private/STutorialsBrowser.cpp b/Engine/Source/Editor/IntroTutorials/Private/STutorialsBrowser.cpp
--- a/Engine/Source/Editor/IntroTutorials/Private/STutorialsBrowser.cpp
+++ b/Engine/Source/Editor/IntroTutorials/Private/STutorialsBrowser.cpp
@@ -204,10 +204,7 @@ public:
 		SubCategories.Add(InSubCategory);
 	}
 
-	void AddTutorial(TSharedPtr<FTutorialListEntry_Tutorial> InTutorial)
-	{
-		Tutorials.Add(InTutorial);
-	}
+	void AddTutorial(TSharedPtr<FTutorialListEntry_Tutorial> InTutorial);
 
 	FReply OnClicked() const
 	{
@@ -626,6 +623,11 @@ void STutorialsBrowser::Construct(const FArguments& InArgs)
 	RebuildCrumbs();
 }
 
+inline void FTutorialListEntry_Category::AddTutorial(TSharedPtr<FTutorialListEntry_Tutorial> InTutorial)
+{
+	Tutorials.Add(InTutorial);
+}
+
 EActiveTimerReturnType STutorialsBrowser::TriggerReloadTutorials( double InCurrentTime, float InDeltaTime )
 {
 	if (bNeedsRefresh)
diff --git a/Engine/Source/Editor/IntroTutorials/Private/TutorialText.cpp b/Engine/Source/Editor/IntroTutorials/Private/TutorialText.cpp
--- a/Engine/Source/Editor/IntroTutorials/Private/TutorialText.cpp
+++ b/Engine/Source/Editor/IntroTutorials/Private/TutorialText.cpp
@@ -326,45 +326,47 @@ void FTutorialText::Initialize()
 {
 	if(HyperlinkDescs.Num() == 0)
 	{
-		HyperlinkDescs.Add(MakeShareable(new FHyperlinkTypeDesc(
-			EHyperlinkType::Browser, 
-			LOCTEXT("BrowserLinkTypeLabel", "URL"), 
-			LOCTEXT("BrowserLinkTypeTooltip", "A link that opens a browser window (e.g. http://www.unrealengine.com)"),
-			TEXT("browser"),
-			FSlateHyperlinkRun::FOnClick::CreateStatic(&OnBrowserLinkClicked))));
-
-		HyperlinkDescs.Add(MakeShareable(new FHyperlinkTypeDesc(
-			EHyperlinkType::UDN, 
-			LOCTEXT("UDNLinkTypeLabel", "UDN"), 
-			LOCTEXT("UDNLinkTypeTooltip", "A link that opens some UDN documentation (e.g. /Engine/Blueprints/UserGuide/Types/ClassBlueprint)"),
-			TEXT("udn"),
-			FSlateHyperlinkRun::FOnClick::CreateStatic(&OnDocLinkClicked), 
-			FSlateHyperlinkRun::FOnGetTooltipText(),
-			FSlateHyperlinkRun::FOnGenerateTooltip::CreateStatic(&OnGenerateDocTooltip))));
-
-		HyperlinkDescs.Add(MakeShareable(new FHyperlinkTypeDesc(
-			EHyperlinkType::Asset, 
-			LOCTEXT("AssetLinkTypeLabel", "Asset"), 
-			LOCTEXT("AssetLinkTypeTooltip", "A link that opens an asset (e.g. /Game/StaticMeshes/SphereMesh.SphereMesh)"),
-			TEXT("asset"),
-			FSlateHyperlinkRun::FOnClick::CreateStatic(&OnAssetLinkClicked),
-			FSlateHyperlinkRun::FOnGetTooltipText::CreateStatic(&OnGetAssetTooltipText))));
-
-		HyperlinkDescs.Add(MakeShareable(new FHyperlinkTypeDesc(
-			EHyperlinkType::Code, 
-			LOCTEXT("CodeLinkTypeLabel", "Code"), 
-			LOCTEXT("CodeLinkTypeTooltip", "A link that opens code in your selected IDE.\nFor example: [PROJECTPATH]/Private/SourceFile.cpp,1,1.\nThe numbers correspond to line number and column number.\nYou can use [PROJECT], [PROJECTPATH] and [ENGINEPATH] tags to make paths."),
-			TEXT("code"),
-			FSlateHyperlinkRun::FOnClick::CreateStatic(&OnCodeLinkClicked),
-			FSlateHyperlinkRun::FOnGetTooltipText::CreateStatic(&OnGetCodeTooltipText))));
-
-		HyperlinkDescs.Add(MakeShareable(new FHyperlinkTypeDesc(
-			EHyperlinkType::Tutorial, 
-			LOCTEXT("TutorialLinkTypeLabel", "Tutorial"), 
-			LOCTEXT("TutorialLinkTypeTooltip", "A type of asset link that opens another tutorial, e.g. /Game/Tutorials/StaticMeshTutorial.StaticMeshTutorial"),
-			TEXT("tutorial"),
-			FSlateHyperlinkRun::FOnClick::CreateStatic(&OnTutorialLinkClicked),
-			FSlateHyperlinkRun::FOnGetTooltipText::CreateStatic(&OnGetTutorialTooltipText))));
+		HyperlinkDescs = {
+			MakeShareable(new FHyperlinkTypeDesc(
+				EHyperlinkType::Browser, 
+				LOCTEXT("BrowserLinkTypeLabel", "URL"), 
+				LOCTEXT("BrowserLinkTypeTooltip", "A link that opens a browser window (e.g. http://www.unrealengine.com)"),
+				TEXT("browser"),
+				FSlateHyperlinkRun::FOnClick::CreateStatic(&OnBrowserLinkClicked))),
+
+			MakeShareable(new FHyperlinkTypeDesc(
+				EHyperlinkType::UDN, 
+				LOCTEXT("UDNLinkTypeLabel", "UDN"), 
+				LOCTEXT("UDNLinkTypeTooltip", "A link that opens some UDN documentation (e.g. /Engine/Blueprints/UserGuide/Types/ClassBlueprint)"),
+				TEXT("udn"),
+				FSlateHyperlinkRun::FOnClick::CreateStatic(&OnDocLinkClicked), 
+				FSlateHyperlinkRun::FOnGetTooltipText(),
+				FSlateHyperlinkRun::FOnGenerateTooltip::CreateStatic(&OnGenerateDocTooltip))),
+
+			MakeShareable(new FHyperlinkTypeDesc(
+				EHyperlinkType::Asset, 
+				LOCTEXT("AssetLinkTypeLabel", "Asset"), 
+				LOCTEXT("AssetLinkTypeTooltip", "A link that opens an asset (e.g. /Game/StaticMeshes/SphereMesh.SphereMesh)"),
+				TEXT("asset"),
+				FSlateHyperlinkRun::FOnClick::CreateStatic(&OnAssetLinkClicked),
+				FSlateHyperlinkRun::FOnGetTooltipText::CreateStatic(&OnGetAssetTooltipText))),
+
+			MakeShareable(new FHyperlinkTypeDesc(
+				EHyperlinkType::Code, 
+				LOCTEXT("CodeLinkTypeLabel", "Code"), 
+				LOCTEXT("CodeLinkTypeTooltip", "A link that opens code in your selected IDE.\nFor example: [PROJECTPATH]/Private/SourceFile.cpp,1,1.\nThe numbers correspond to line number and column number.\nYou can use [PROJECT], [PROJECTPATH] and [ENGINEPATH] tags to make paths."),
+				TEXT("code"),
+				FSlateHyperlinkRun::FOnClick::CreateStatic(&OnCodeLinkClicked),
+				FSlateHyperlinkRun::FOnGetTooltipText::CreateStatic(&OnGetCodeTooltipText))),
+
+			MakeShareable(new FHyperlinkTypeDesc(
+				EHyperlinkType::Tutorial, 
+				LOCTEXT("TutorialLinkTypeLabel", "Tutorial"), 
+				LOCTEXT("TutorialLinkTypeTooltip", "A type of asset link that opens another tutorial, e.g. /Game/Tutorials/StaticMeshTutorial.StaticMeshTutorial"),
+				TEXT("tutorial"),
+				FSlateHyperlinkRun::FOnClick::CreateStatic(&OnTutorialLinkClicked),
+				FSlateHyperlinkRun::FOnGetTooltipText::CreateStatic(&OnGetTutorialTooltipText)))
+		};
 	}
 }
 
diff --git a/Engine/Source/Editor/LandscapeEditor/Private/LandscapeEdModeSplineTools.cpp b/Engine/Source/Editor/LandscapeEditor/Private/LandscapeEdModeSplineTools.cpp
--- a/Engine/Source/Editor/LandscapeEditor/Private/LandscapeEdModeSplineTools.cpp
+++ b/Engine/Source/Editor/LandscapeEditor/Private/LandscapeEdModeSplineTools.cpp
@@ -648,6 +648,7 @@ public:
 					if (ToLandscape->SplineComponent == nullptr)
 					{
 						CreateSplineComponent(ToLandscape, FromProxy->SplineComponent->RelativeScale3D);
+						check(ToLandscape->SplineComponent);
 					}
 					ToLandscape->SplineComponent->Modify();
 
@@ -737,6 +738,7 @@ public:
 					if (ToLandscape->SplineComponent == nullptr)
 					{
 						CreateSplineComponent(ToLandscape, FromProxy->SplineComponent->RelativeScale3D);
+						check(ToLandscape->SplineComponent);
 					}
 					ToLandscape->SplineComponent->Modify();
 
@@ -858,6 +860,7 @@ public:
 				if (!Landscape->SplineComponent)
 				{
 					CreateSplineComponent(Landscape, FVector(1.0f) / Landscape->GetRootComponent()->RelativeScale3D);
+					check(Landscape->SplineComponent);
 				}
 				SplinesComponent = Landscape->SplineComponent;
 			}
@@ -1820,6 +1823,7 @@ public:
 		if (!Landscape->SplineComponent)
 		{
 			CreateSplineComponent(Landscape, FVector(1.0f) / Landscape->GetRootComponent()->RelativeScale3D);
+			check(Landscape->SplineComponent);
 		}
 		Landscape->SplineComponent->Modify();
 
diff --git a/Engine/Source/Editor/LandscapeEditor/Private/LandscapeEdModeTools.h b/Engine/Source/Editor/LandscapeEditor/Private/LandscapeEdModeTools.h
--- a/Engine/Source/Editor/LandscapeEditor/Private/LandscapeEdModeTools.h
+++ b/Engine/Source/Editor/LandscapeEditor/Private/LandscapeEdModeTools.h
@@ -222,6 +222,7 @@ inline void LowPassFilter(int32 X1, int32 Y1, int32 X2, int32 Y2, FLandscapeBrus
 			float Ratio = 1.0f - DetailScale;
 			float Dist = FMath::Min<float>((Dims[0] * Ratio)*(Dims[0] * Ratio), (Dims[1] * Ratio)*(Dims[1] * Ratio));
 			float Filter = 1.0 / (1.0 + DistFromCenter / Dist);
+			CA_SUPPRESS(6385);
 			out[X + Y*Dims[1]].r *= Filter;
 			out[X + Y*Dims[1]].i *= Filter;
 		}
diff --git a/Engine/Source/Editor/LandscapeEditor/Public/LandscapeToolInterface.h b/Engine/Source/Editor/LandscapeEditor/Public/LandscapeToolInterface.h
--- a/Engine/Source/Editor/LandscapeEditor/Public/LandscapeToolInterface.h
+++ b/Engine/Source/Editor/LandscapeEditor/Public/LandscapeToolInterface.h
@@ -70,7 +70,7 @@ public:
 		return BrushAlpha.GetData() + (Position.Y - Bounds.Min.Y) * Bounds.Width() + (Position.X - Bounds.Min.X);
 	}
 
-	FORCEINLINE_EXPLICIT_OPERATOR_BOOL() const
+	FORCEINLINE explicit operator bool() const
 	{
 		return BrushAlpha.Num() != 0;
 	}
diff --git a/Engine/Source/Editor/Persona/Private/SAnimationBlendSpace.cpp b/Engine/Source/Editor/Persona/Private/SAnimationBlendSpace.cpp
--- a/Engine/Source/Editor/Persona/Private/SAnimationBlendSpace.cpp
+++ b/Engine/Source/Editor/Persona/Private/SAnimationBlendSpace.cpp
@@ -854,8 +854,6 @@ int32 SBlendSpaceGridWidget::OnPaint( const FPaintArgs& Args, const FGeometry& A
 
 	FDrawLines LineBatcher;
 
-	TArray<FVector2D> LinePoints;
-
 	int32 DefaultLayer = LayerId;
 	int32 GridLayer = LayerId+1;
 	int32 HighlightLayer = LayerId+2;
@@ -869,11 +867,13 @@ int32 SBlendSpaceGridWidget::OnPaint( const FPaintArgs& Args, const FGeometry& A
 
 	// draw the borderline for grid
 	{
-		LinePoints.Add( FVector2D(WindowRect.Left, WindowRect.Top) );
-		LinePoints.Add( FVector2D(WindowRect.Left, WindowRect.Bottom) );
-		LinePoints.Add( FVector2D(WindowRect.Right, WindowRect.Bottom) );
-		LinePoints.Add( FVector2D(WindowRect.Right, WindowRect.Top));
-		LinePoints.Add( FVector2D(WindowRect.Left, WindowRect.Top) );
+		TArray<FVector2D> LinePoints = {
+			FVector2D(WindowRect.Left, WindowRect.Top),
+			FVector2D(WindowRect.Left, WindowRect.Bottom),
+			FVector2D(WindowRect.Right, WindowRect.Bottom),
+			FVector2D(WindowRect.Right, WindowRect.Top),
+			FVector2D(WindowRect.Left, WindowRect.Top)
+		};
 
 		FSlateDrawElement::MakeLines( 
 			OutDrawElements,
diff --git a/Engine/Source/Editor/Persona/Private/STrack.cpp b/Engine/Source/Editor/Persona/Private/STrack.cpp
--- a/Engine/Source/Editor/Persona/Private/STrack.cpp
+++ b/Engine/Source/Editor/Persona/Private/STrack.cpp
@@ -61,6 +61,12 @@ TSharedRef<FTrackNodeDragDropOp> FTrackNodeDragDropOp::New(TSharedRef<STrackNode
 	return Operation;
 }
 
+TSharedPtr<SWidget> FTrackNodeDragDropOp::GetDefaultDecorator() const
+{
+	return OriginalTrackNode.Pin();
+}
+
+
 //////////////////////////////////////////////////////////////////////////
 // STrackNode
 void STrackNode::Construct(const FArguments& InArgs)
diff --git a/Engine/Source/Editor/Persona/Private/STrack.h b/Engine/Source/Editor/Persona/Private/STrack.h
--- a/Engine/Source/Editor/Persona/Private/STrack.h
+++ b/Engine/Source/Editor/Persona/Private/STrack.h
@@ -40,10 +40,7 @@ public:
 	static TSharedRef<FTrackNodeDragDropOp> New(TSharedRef<STrackNode> TrackNode, const FVector2D &CursorPosition, const FVector2D &ScreenPositionOfNode);
 
 	/** Gets the widget that will serve as the decorator unless overridden. If you do not override, you will have no decorator */
-	virtual TSharedPtr<SWidget> GetDefaultDecorator() const override
-	{
-		return OriginalTrackNode.Pin();
-	}
+	virtual TSharedPtr<SWidget> GetDefaultDecorator() const override;
 
 protected:
 	/** The window that shows hover text */
diff --git a/Engine/Source/Editor/PropertyEditor/Private/UserInterface/PropertyEditor/SPropertyEditorNumeric.h b/Engine/Source/Editor/PropertyEditor/Private/UserInterface/PropertyEditor/SPropertyEditorNumeric.h
--- a/Engine/Source/Editor/PropertyEditor/Private/UserInterface/PropertyEditor/SPropertyEditorNumeric.h
+++ b/Engine/Source/Editor/PropertyEditor/Private/UserInterface/PropertyEditor/SPropertyEditorNumeric.h
@@ -31,58 +31,66 @@ public:
 
 		if(!Property->IsA(UFloatProperty::StaticClass()) && Property->HasMetaData(PropertyEditorConstants::MD_Bitmask))
 		{
-			const int32 BitmaskBitCount = sizeof(NumericType) << 3;
+			auto CreateBitmaskFlagsArray = [](const UProperty* Prop)
+			{
+				const int32 BitmaskBitCount = sizeof(NumericType) << 3;
 
-			BitmaskFlags.Empty(BitmaskBitCount);
+				TArray<FBitmaskFlagInfo> Result;
+				Result.Empty(BitmaskBitCount);
 
-			const UEnum* BitmaskEnum = nullptr;
-			const FString& BitmaskEnumName = Property->GetMetaData(PropertyEditorConstants::MD_BitmaskEnum);
-			if (!BitmaskEnumName.IsEmpty())
-			{
-				// @TODO: Potentially replace this with a parameter passed in from a member variable on the UProperty (e.g. UByteProperty::Enum)
-				BitmaskEnum = FindObject<UEnum>(ANY_PACKAGE, *BitmaskEnumName);
-			}
+				const UEnum* BitmaskEnum = nullptr;
+				const FString& BitmaskEnumName = Prop->GetMetaData(PropertyEditorConstants::MD_BitmaskEnum);
+				if (!BitmaskEnumName.IsEmpty())
+				{
+					// @TODO: Potentially replace this with a parameter passed in from a member variable on the UProperty (e.g. UByteProperty::Enum)
+					BitmaskEnum = FindObject<UEnum>(ANY_PACKAGE, *BitmaskEnumName);
+				}
 
-			if (BitmaskEnum)
-			{
-				for (int32 BitmaskEnumIndex = 0; BitmaskEnumIndex < BitmaskEnum->NumEnums() - 1; ++BitmaskEnumIndex)
+				if (BitmaskEnum)
 				{
-					int8 BitmaskFlagIndex = BitmaskEnum->GetValueByIndex(BitmaskEnumIndex);
-					if (BitmaskFlagIndex >= 0 && BitmaskFlagIndex < BitmaskBitCount)
+					for (int32 BitmaskEnumIndex = 0; BitmaskEnumIndex < BitmaskEnum->NumEnums() - 1; ++BitmaskEnumIndex)
 					{
-						FBitmaskFlagInfo* BitmaskFlag = new(BitmaskFlags) FBitmaskFlagInfo();
+						int8 BitmaskFlagIndex = BitmaskEnum->GetValueByIndex(BitmaskEnumIndex);
+						if (BitmaskFlagIndex >= 0 && BitmaskFlagIndex < BitmaskBitCount)
+						{
+							Result.Emplace();
+							FBitmaskFlagInfo* BitmaskFlag = &Result.Last();
 
-						BitmaskFlag->Value = TBitmaskValueHelpers<NumericType>::LeftShift(static_cast<NumericType>(1), static_cast<int32>(BitmaskFlagIndex));
+							BitmaskFlag->Value = TBitmaskValueHelpers<NumericType>::LeftShift(static_cast<NumericType>(1), static_cast<int32>(BitmaskFlagIndex));
 
-						BitmaskFlag->DisplayName = BitmaskEnum->GetDisplayNameText(BitmaskEnumIndex);
-						if (BitmaskFlag->DisplayName.IsEmpty())
-						{
-							BitmaskFlag->DisplayName = FText::FromString(BitmaskEnum->GetEnumName(BitmaskEnumIndex));
-						}
+							BitmaskFlag->DisplayName = BitmaskEnum->GetDisplayNameText(BitmaskEnumIndex);
+							if (BitmaskFlag->DisplayName.IsEmpty())
+							{
+								BitmaskFlag->DisplayName = FText::FromString(BitmaskEnum->GetEnumName(BitmaskEnumIndex));
+							}
 
-						BitmaskFlag->ToolTipText = BitmaskEnum->GetToolTipText(BitmaskEnumIndex);
-						if (BitmaskFlag->ToolTipText.IsEmpty())
-						{
-							BitmaskFlag->ToolTipText = FText::Format(LOCTEXT("BitmaskDefaultFlagToolTipText", "Toggle {0} on/off"), BitmaskFlag->DisplayName);
+							BitmaskFlag->ToolTipText = BitmaskEnum->GetToolTipText(BitmaskEnumIndex);
+							if (BitmaskFlag->ToolTipText.IsEmpty())
+							{
+								BitmaskFlag->ToolTipText = FText::Format(LOCTEXT("BitmaskDefaultFlagToolTipText", "Toggle {0} on/off"), BitmaskFlag->DisplayName);
+							}
 						}
 					}
 				}
-			}
-			else
-			{
-				for (int32 BitmaskFlagIndex = 0; BitmaskFlagIndex < BitmaskBitCount; ++BitmaskFlagIndex)
+				else
 				{
-					FBitmaskFlagInfo* BitmaskFlag = new(BitmaskFlags) FBitmaskFlagInfo();
+					for (int32 BitmaskFlagIndex = 0; BitmaskFlagIndex < BitmaskBitCount; ++BitmaskFlagIndex)
+					{
+						Result.Emplace();
+						FBitmaskFlagInfo* BitmaskFlag = &Result.Last();
 
-					BitmaskFlag->Value = TBitmaskValueHelpers<NumericType>::LeftShift(static_cast<NumericType>(1), BitmaskFlagIndex);
-					BitmaskFlag->DisplayName = FText::Format(LOCTEXT("BitmaskDefaultFlagDisplayName", "Flag {0}"), FText::AsNumber(BitmaskFlagIndex + 1));
-					BitmaskFlag->ToolTipText = FText::Format(LOCTEXT("BitmaskDefaultFlagToolTipText", "Toggle {0} on/off"), BitmaskFlag->DisplayName);
+						BitmaskFlag->Value = TBitmaskValueHelpers<NumericType>::LeftShift(static_cast<NumericType>(1), BitmaskFlagIndex);
+						BitmaskFlag->DisplayName = FText::Format(LOCTEXT("BitmaskDefaultFlagDisplayName", "Flag {0}"), FText::AsNumber(BitmaskFlagIndex + 1));
+						BitmaskFlag->ToolTipText = FText::Format(LOCTEXT("BitmaskDefaultFlagToolTipText", "Toggle {0} on/off"), BitmaskFlag->DisplayName);
+					}
 				}
-			}
+
+				return Result;
+			};
 
 			const FComboBoxStyle& ComboBoxStyle = FCoreStyle::Get().GetWidgetStyle< FComboBoxStyle >("ComboBox");
 
-			const auto& GetComboButtonText = [this]() -> FText
+			const auto& GetComboButtonText = [this, CreateBitmaskFlagsArray, Property]() -> FText
 			{
 				TOptional<NumericType> Value = OnGetValue();
 				if (Value.IsSet())
@@ -96,6 +104,7 @@ public:
 						}
 						else
 						{
+							TArray<FBitmaskFlagInfo> BitmaskFlags = CreateBitmaskFlagsArray(Property);
 							for (int i = 0; i < BitmaskFlags.Num(); ++i)
 							{
 								if (TBitmaskValueHelpers<NumericType>::BitwiseAND(BitmaskValue, BitmaskFlags[i].Value))
@@ -124,10 +133,11 @@ public:
 				.Font(InArgs._Font)
 				.Text_Lambda(GetComboButtonText)
 			]
-			.OnGetMenuContent_Lambda([this]()
+			.OnGetMenuContent_Lambda([this, CreateBitmaskFlagsArray, Property]()
 			{
 				FMenuBuilder MenuBuilder(false, nullptr);
 
+				TArray<FBitmaskFlagInfo> BitmaskFlags = CreateBitmaskFlagsArray(Property);
 				for (int i = 0; i < BitmaskFlags.Num(); ++i)
 				{
 					MenuBuilder.AddMenuEntry(
@@ -136,7 +146,7 @@ public:
 						FSlateIcon(),
 						FUIAction
 						(
-							FExecuteAction::CreateLambda([this, i]()
+							FExecuteAction::CreateLambda([this, i, BitmaskFlags]()
 							{
 								TOptional<NumericType> Value = OnGetValue();
 								if (Value.IsSet())
@@ -145,7 +155,7 @@ public:
 								}
 							}),
 							FCanExecuteAction(),
-							FIsActionChecked::CreateLambda([this, i]() -> bool
+							FIsActionChecked::CreateLambda([this, i, BitmaskFlags]() -> bool
 							{
 								TOptional<NumericType> Value = OnGetValue();
 								if (Value.IsSet())
@@ -544,9 +554,6 @@ private:
 
 	/** True if the slider is being used to change the value of the property */
 	bool bIsUsingSlider;
-
-	/** Internal array of per-flag info, used for bitmask editing. */
-	TArray<FBitmaskFlagInfo> BitmaskFlags;
 };
 
 #undef LOCTEXT_NAMESPACE
\ No newline at end of file
diff --git a/Engine/Source/Editor/Sequencer/Private/GroupedKeyArea.cpp b/Engine/Source/Editor/Sequencer/Private/GroupedKeyArea.cpp
--- a/Engine/Source/Editor/Sequencer/Private/GroupedKeyArea.cpp
+++ b/Engine/Source/Editor/Sequencer/Private/GroupedKeyArea.cpp
@@ -48,9 +48,12 @@ void FGroupedKeyCollection::InitializeExplicit(const TArray<FSequencerDisplayNod
 		else if (Node->GetType() == ESequencerNode::Track)
 		{
 			KeyAreaNode = static_cast<const FSequencerTrackNode*>(Node)->GetTopLevelKeyNode().Get();
+			if (!KeyAreaNode)
+			{
+				continue;
+			}
 		}
-
-		if (!KeyAreaNode)
+		else
 		{
 			continue;
 		}
diff --git a/Engine/Source/Editor/Sequencer/Private/Sequencer.cpp b/Engine/Source/Editor/Sequencer/Private/Sequencer.cpp
--- a/Engine/Source/Editor/Sequencer/Private/Sequencer.cpp
+++ b/Engine/Source/Editor/Sequencer/Private/Sequencer.cpp
@@ -533,6 +533,12 @@ void FSequencer::PostTickRenderStateFixup()
 }
 
 
+TSharedRef<SWidget> FSequencer::GetSequencerWidget() const
+{
+	return SequencerWidget.ToSharedRef();
+}
+
+
 UMovieSceneSequence* FSequencer::GetRootMovieSceneSequence() const
 {
 	if (SequenceInstanceStack.Num())
@@ -5430,7 +5436,7 @@ void FixSceneRangeTiming( UMovieScene* MovieScene, float FrameInterval )
 	float UpperBoundValue = SceneRange.GetUpperBoundValue();
 	float SnappedUpperBoundValue = SnapTime( UpperBoundValue, FrameInterval );
 
-	if ( SnappedLowerBoundValue != LowerBoundValue || SnappedLowerBoundValue != LowerBoundValue)
+	if ( SnappedLowerBoundValue != LowerBoundValue || SnappedUpperBoundValue != UpperBoundValue)
 	{
 		MovieScene->SetPlaybackRange( SnappedLowerBoundValue, SnappedUpperBoundValue );
 	}
diff --git a/Engine/Source/Editor/Sequencer/Private/Sequencer.h b/Engine/Source/Editor/Sequencer/Private/Sequencer.h
--- a/Engine/Source/Editor/Sequencer/Private/Sequencer.h
+++ b/Engine/Source/Editor/Sequencer/Private/Sequencer.h
@@ -491,7 +491,7 @@ public:
 	//~ ISequencer Interface
 
 	virtual void Close() override;
-	virtual TSharedRef<SWidget> GetSequencerWidget() const override { return SequencerWidget.ToSharedRef(); }
+	virtual TSharedRef<SWidget> GetSequencerWidget() const override;
 	virtual UMovieSceneSequence* GetRootMovieSceneSequence() const override;
 	virtual UMovieSceneSequence* GetFocusedMovieSceneSequence() const override;
 	virtual void ResetToNewRootSequence(UMovieSceneSequence& NewSequence) override;
diff --git a/Engine/Source/Editor/StaticMeshEditor/Private/StaticMeshAutomationTests.cpp b/Engine/Source/Editor/StaticMeshEditor/Private/StaticMeshAutomationTests.cpp
--- a/Engine/Source/Editor/StaticMeshEditor/Private/StaticMeshAutomationTests.cpp
+++ b/Engine/Source/Editor/StaticMeshEditor/Private/StaticMeshAutomationTests.cpp
@@ -179,7 +179,8 @@ bool FStaticMeshEditorTest::RunTest(const FString& Parameters)
 		FSlateApplication::Get().GetAllVisibleWindowsOrdered(AllWindows);
 
 		const FText ExpectedTitle = FText::FromString(LoadedObjectType);
-		if (auto* Window = AllWindows.FindByPredicate([&](const TSharedRef<SWindow>& In) { return In->GetTitle().EqualTo(ExpectedTitle); }))
+		TSharedRef<SWindow>* Window = AllWindows.FindByPredicate([&](const TSharedRef<SWindow>& In) { return In->GetTitle().EqualTo(ExpectedTitle); });
+		if (Window)
 		{
 			WindowParameters.CurrentWindow = *Window;
 		}
diff --git a/Engine/Source/Editor/UMGEditor/Private/WidgetBlueprint.cpp b/Engine/Source/Editor/UMGEditor/Private/WidgetBlueprint.cpp
--- a/Engine/Source/Editor/UMGEditor/Private/WidgetBlueprint.cpp
+++ b/Engine/Source/Editor/UMGEditor/Private/WidgetBlueprint.cpp
@@ -655,9 +655,4 @@ bool UWidgetBlueprint::NeedsLoadForClient() const
 	return false;
 }
 
-bool UWidgetBlueprint::NeedsLoadForServer() const
-{
-	return false;
-}
-
 #undef LOCTEXT_NAMESPACE 
\ No newline at end of file
diff --git a/Engine/Source/Editor/UMGEditor/Public/WidgetBlueprint.h b/Engine/Source/Editor/UMGEditor/Public/WidgetBlueprint.h
--- a/Engine/Source/Editor/UMGEditor/Public/WidgetBlueprint.h
+++ b/Engine/Source/Editor/UMGEditor/Public/WidgetBlueprint.h
@@ -177,7 +177,6 @@ public:
 
 	//~ Begin UObject Interface
 	virtual bool NeedsLoadForClient() const override;
-	virtual bool NeedsLoadForServer() const override;
 	//~ End UObject Interface
 
 public:
diff --git a/Engine/Source/Editor/UnrealEd/Classes/Editor/GroupActor.h b/Engine/Source/Editor/UnrealEd/Classes/Editor/GroupActor.h
--- a/Engine/Source/Editor/UnrealEd/Classes/Editor/GroupActor.h
+++ b/Engine/Source/Editor/UnrealEd/Classes/Editor/GroupActor.h
@@ -35,10 +35,6 @@ class AGroupActor : public AActor
 	{ 
 		return false; 
 	}
-	virtual bool NeedsLoadForServer() const override
-	{ 
-		return false;
-	}
 	//~ End UObject Interface
 
 	//~ Begin AActor Interface
diff --git a/Engine/Source/Editor/UnrealEd/Private/CookOnTheFlyServer.cpp b/Engine/Source/Editor/UnrealEd/Private/CookOnTheFlyServer.cpp
--- a/Engine/Source/Editor/UnrealEd/Private/CookOnTheFlyServer.cpp
+++ b/Engine/Source/Editor/UnrealEd/Private/CookOnTheFlyServer.cpp
@@ -3453,8 +3453,9 @@ void UCookOnTheFlyServer::CleanSandbox( const bool bIterative )
 
 					FString UncookedFilename = ConvertCookedPathToUncookedPath(CookedFilename);
 					
-					FDateTime DependentTimestamp;
 					FName StandardUnCookedFileFName = FName(*UncookedFilename);
+					
+					FDateTime DependentTimestamp = IFileManager::Get().GetTimeStamp(*StandardUnCookedFileFName.ToString());
 #if 1 // asset registry dependency checking instead of using DependencyInfo module
 					//////////////////////////////////////////////////////////////////////////
 
diff --git a/Engine/Source/Editor/UnrealEd/Private/EditorBsp.cpp b/Engine/Source/Editor/UnrealEd/Private/EditorBsp.cpp
--- a/Engine/Source/Editor/UnrealEd/Private/EditorBsp.cpp
+++ b/Engine/Source/Editor/UnrealEd/Private/EditorBsp.cpp
@@ -313,9 +313,16 @@ void UEditorEngine::bspMergeCoplanars( UModel* Model, bool RemapLinks, bool Merg
 	}
 	Model->Polys->Element.RemoveAt( j, Model->Polys->Element.Num()-j );
 	if( RemapLinks )
+	{
 		for( int32 i=0; i<Model->Polys->Element.Num(); i++ )
+		{
 			if( Model->Polys->Element[i].iLink != INDEX_NONE )
+			{
+				CA_SUPPRESS(6001); // warning C6001: Using uninitialized memory '*Remap'.
 				Model->Polys->Element[i].iLink = Remap[Model->Polys->Element[i].iLink];
+			}
+		}
+	}
 // 	UE_LOG(LogEditorBsp, Log,  TEXT("BspMergeCoplanars reduced %i->%i"), OriginalNum, Model->Polys->Element.Num() );
 	Mark2.Pop();
 }
@@ -1702,13 +1709,23 @@ void MergeNearPoints( UModel *Model, float Dist )
 
 	// Remap VertPool.
 	for( int32 i=0; i<Model->Verts.Num(); i++ )
+	{
 		if( Model->Verts[i].pVertex>=0 && Model->Verts[i].pVertex<Model->Points.Num() )
+		{
+			CA_SUPPRESS(6001); // warning C6001: Using uninitialized memory '*PointRemap'.
 			Model->Verts[i].pVertex = PointRemap[Model->Verts[i].pVertex];
+		}
+	}
 
 	// Remap Surfs.
 	for( int32 i=0; i<Model->Surfs.Num(); i++ )
+	{
 		if( Model->Surfs[i].pBase>=0 && Model->Surfs[i].pBase<Model->Points.Num() )
+		{
+			CA_SUPPRESS(6001); // warning C6001: Using uninitialized memory '*PointRemap'.
 			Model->Surfs[i].pBase = PointRemap[Model->Surfs[i].pBase];
+		}
+	}
 
 	// Remove duplicate points from nodes.
 	for( int32 i=0; i<Model->Nodes.Num(); i++ )
diff --git a/Engine/Source/Editor/UnrealEd/Private/Fbx/FbxMainImport.cpp b/Engine/Source/Editor/UnrealEd/Private/Fbx/FbxMainImport.cpp
--- a/Engine/Source/Editor/UnrealEd/Private/Fbx/FbxMainImport.cpp
+++ b/Engine/Source/Editor/UnrealEd/Private/Fbx/FbxMainImport.cpp
@@ -1061,14 +1061,16 @@ FName FFbxImporter::MakeNameForMesh(FString InName, FbxObject* FbxObject)
 	FName OutputName;
 
 	//Cant name the mesh if the object is null and there InName arguments is None.
-	check(FbxObject != NULL || InName != FString("None"))
+	check(FbxObject != nullptr || InName != TEXT("None"))
 
-	if ((ImportOptions->bUsedAsFullName || FbxObject == NULL) && InName != FString("None"))
+	if ((ImportOptions->bUsedAsFullName || FbxObject == nullptr) && InName != TEXT("None"))
 	{
 		OutputName = *InName;
 	}
 	else
 	{
+		check(FbxObject);
+
 		char Name[MAX_SPRINTF];
 		int SpecialChars[] = {'.', ',', '/', '`', '%'};
 
@@ -1077,14 +1079,14 @@ FName FFbxImporter::MakeNameForMesh(FString InName, FbxObject* FbxObject)
 		for ( int32 i = 0; i < 5; i++ )
 		{
 			char* CharPtr = Name;
-			while ( (CharPtr = FCStringAnsi::Strchr(CharPtr,SpecialChars[i])) != NULL )
+			while ( (CharPtr = FCStringAnsi::Strchr(CharPtr,SpecialChars[i])) != nullptr )
 			{
 				CharPtr[0] = '_';
 			}
 		}
 
 		// for mesh, replace ':' with '_' because Unreal doesn't support ':' in mesh name
-		char* NewName = NULL;
+		char* NewName = nullptr;
 		NewName = FCStringAnsi::Strchr (Name, ':');
 
 		if (NewName)
@@ -1732,8 +1734,8 @@ void FFbxImporter::RecursiveFindFbxSkelMesh(FbxNode* Node, TArray< TArray<FbxNod
 
 void FFbxImporter::RecursiveFindRigidMesh(FbxNode* Node, TArray< TArray<FbxNode*>* >& outSkelMeshArray, TArray<FbxNode*>& SkeletonArray, bool ExpandLOD)
 {
-	bool RigidNodeFound = false;
-	FbxNode* RigidMeshNode = NULL;
+	bool bRigidNodeFound = false;
+	FbxNode* RigidMeshNode = nullptr;
 
 	DEBUG_FBX_NODE("", Node);
 
@@ -1743,7 +1745,7 @@ void FFbxImporter::RecursiveFindRigidMesh(FbxNode* Node, TArray< TArray<FbxNode*
 		if (Node->GetMesh()->GetDeformerCount(FbxDeformer::eSkin) == 0 )
 		{
 			RigidMeshNode = Node;
-			RigidNodeFound = true;
+			bRigidNodeFound = true;
 		}
 	}
 	else if (Node->GetNodeAttribute() && Node->GetNodeAttribute()->GetAttributeType() == FbxNodeAttribute::eLODGroup)
@@ -1755,11 +1757,11 @@ void FFbxImporter::RecursiveFindRigidMesh(FbxNode* Node, TArray< TArray<FbxNode*
 		{
 			if (FirstLOD->GetMesh()->GetDeformerCount(FbxDeformer::eSkin) == 0 )
 			{
-				RigidNodeFound = true;
+				bRigidNodeFound = true;
 			}
 		}
 
-		if (RigidNodeFound)
+		if (bRigidNodeFound)
 		{
 			if (ExpandLOD)
 			{
@@ -1773,7 +1775,7 @@ void FFbxImporter::RecursiveFindRigidMesh(FbxNode* Node, TArray< TArray<FbxNode*
 		}
 	}
 
-	if (RigidMeshNode)
+	if (bRigidNodeFound)
 	{
 		// find root skeleton
 		FbxNode* Link = GetRootSkeleton(RigidMeshNode);
@@ -2050,30 +2052,29 @@ void FFbxImporter::CheckSmoothingInfo(FbxMesh* FbxMesh)
 //-------------------------------------------------------------------------
 FbxNode* FFbxImporter::RetrieveObjectFromName(const TCHAR* ObjectName, FbxNode* Root)
 {
-	FbxNode* Result = NULL;
-	
-	if ( Scene != NULL )
+	if (Scene)
 	{
-		if (Root == NULL)
+		if (!Root)
 		{
 			Root = Scene->GetRootNode();
 		}
 
-		for (int32 ChildIndex=0;ChildIndex<Root->GetChildCount() && !Result;++ChildIndex)
+		for (int32 ChildIndex=0;ChildIndex<Root->GetChildCount();++ChildIndex)
 		{
 			FbxNode* Node = Root->GetChild(ChildIndex);
 			FbxMesh* FbxMesh = Node->GetMesh();
 			if (FbxMesh && 0 == FCString::Strcmp(ObjectName,UTF8_TO_TCHAR(Node->GetName())))
 			{
-				Result = Node;
+				return Node;
 			}
-			else
+
+			if (FbxNode* NextNode = RetrieveObjectFromName(ObjectName,Node))
 			{
-				Result = RetrieveObjectFromName(ObjectName,Node);
+				return NextNode;
 			}
 		}
 	}
-	return Result;
+	return nullptr;
 }
 
 } // namespace UnFbx
diff --git a/Engine/Source/Editor/UnrealEd/Private/PackageTools.cpp b/Engine/Source/Editor/UnrealEd/Private/PackageTools.cpp
--- a/Engine/Source/Editor/UnrealEd/Private/PackageTools.cpp
+++ b/Engine/Source/Editor/UnrealEd/Private/PackageTools.cpp
@@ -62,6 +62,10 @@ namespace PackageTools
 		// Assemble a list of packages.  Only show packages that match the current resource type filter.
 		for (UObject* Obj : TObjectRange<UObject>())
 		{
+			// This is here to hopefully catch a bit more info about a spurious in-the-wild problem which ultimately
+			// crashes inside UObjectBaseUtility::GetOutermost(), which is called inside IsObjectBrowsable().
+			checkf(Obj->IsValidLowLevel(), TEXT("GetFilteredPackageList: bad object found, address: %p, name: %s"), Obj, *Obj->GetName());
+
 			// Make sure that we support displaying this object type
 			bool bIsSupported = ObjectTools::IsObjectBrowsable( Obj );
 			if( bIsSupported )
diff --git a/Engine/Source/Editor/UnrealEd/Public/Kismet2/KismetDebugUtilities.h b/Engine/Source/Editor/UnrealEd/Public/Kismet2/KismetDebugUtilities.h
--- a/Engine/Source/Editor/UnrealEd/Public/Kismet2/KismetDebugUtilities.h
+++ b/Engine/Source/Editor/UnrealEd/Public/Kismet2/KismetDebugUtilities.h
@@ -40,10 +40,8 @@ public:
 	/** Advances iterator to the next element in the container. */
 	FObjectsBeingDebuggedIterator& operator++();
 
-	SAFE_BOOL_OPERATORS(FObjectsBeingDebuggedIterator)
-
 	/** conversion to "bool" returning true if the iterator has not reached the last element. */
-	FORCEINLINE_EXPLICIT_OPERATOR_BOOL() const
+	FORCEINLINE explicit operator bool() const
 	{ 
 		return IsValid(); 
 	}
@@ -82,10 +80,8 @@ public:
 	/** Advances iterator to the next element in the container. */
 	FBlueprintObjectsBeingDebuggedIterator& operator++();
 
-	SAFE_BOOL_OPERATORS(FBlueprintObjectsBeingDebuggedIterator)
-
 	/** conversion to "bool" returning true if the iterator has not reached the last element. */
-	FORCEINLINE_EXPLICIT_OPERATOR_BOOL() const
+	FORCEINLINE explicit operator bool() const
 	{ 
 		return IsValid(); 
 	}
diff --git a/Engine/Source/Editor/UserFeedback/Private/UserFeedback.cpp b/Engine/Source/Editor/UserFeedback/Private/UserFeedback.cpp
--- a/Engine/Source/Editor/UserFeedback/Private/UserFeedback.cpp
+++ b/Engine/Source/Editor/UserFeedback/Private/UserFeedback.cpp
@@ -210,22 +210,26 @@ private:
 	{
 		// ------------------------------------
 		// Fill the context display names array
-		ContextDisplayNames.Add(LOCTEXT("LevelEditing", "Level Editing"));
-		ContextDisplayNames.Add(LOCTEXT("ContentBrowser", "Content Browser"));
-		ContextDisplayNames.Add(LOCTEXT("PlayInEditor", "Play In Editor"));
-		ContextDisplayNames.Add(LOCTEXT("AssetCreation", "Asset Creation"));
-		ContextDisplayNames.Add(LOCTEXT("DetailsPanel", "Details Panel"));
-		ContextDisplayNames.Add(LOCTEXT("Tutorials", "Tutorials"));
-		ContextDisplayNames.Add(LOCTEXT("Other", "Other"));
+		ContextDisplayNames.Append({
+			LOCTEXT("LevelEditing", "Level Editing"),
+			LOCTEXT("ContentBrowser", "Content Browser"),
+			LOCTEXT("PlayInEditor", "Play In Editor"),
+			LOCTEXT("AssetCreation", "Asset Creation"),
+			LOCTEXT("DetailsPanel", "Details Panel"),
+			LOCTEXT("Tutorials", "Tutorials"),
+			LOCTEXT("Other", "Other")
+		});
 
 		ContextMarkers.FirstEditor = ContextDisplayNames.Num();
-		ContextDisplayNames.Add(NSLOCTEXT("BlueprintEditor", "AppLabel", "Blueprint Editor"));
-		ContextDisplayNames.Add(NSLOCTEXT("Matinee", "AppLabel", "Matinee"));
-		ContextDisplayNames.Add(NSLOCTEXT("StaticMeshEditor", "AppLabel", "StaticMesh Editor"));
-		ContextDisplayNames.Add(NSLOCTEXT("MaterialEditor", "AppLabel", "Material Editor"));
-		ContextDisplayNames.Add(NSLOCTEXT("PhAT", "AppLabel", "PhAT"));
-		ContextDisplayNames.Add(NSLOCTEXT("Cascade", "AppLabel", "Cascade"));
-		ContextDisplayNames.Add(NSLOCTEXT("FPersona", "AppLabel", "Persona"));
+		ContextDisplayNames.Append({
+			NSLOCTEXT("BlueprintEditor", "AppLabel", "Blueprint Editor"),
+			NSLOCTEXT("Matinee", "AppLabel", "Matinee"),
+			NSLOCTEXT("StaticMeshEditor", "AppLabel", "StaticMesh Editor"),
+			NSLOCTEXT("MaterialEditor", "AppLabel", "Material Editor"),
+			NSLOCTEXT("PhAT", "AppLabel", "PhAT"),
+			NSLOCTEXT("Cascade", "AppLabel", "Cascade"),
+			NSLOCTEXT("FPersona", "AppLabel", "Persona")
+		});
 
 		ContextMarkers.Custom = -1;
 
@@ -353,10 +357,11 @@ private:
 		{
 			if (FEngineAnalytics::IsAvailable())
 			{
-				TArray<FAnalyticsEventAttribute> Attributes;
-				Attributes.Add(FAnalyticsEventAttribute(TEXT("Positive"), Mode == EFeedbackMode::Positive));
-				Attributes.Add(FAnalyticsEventAttribute(TEXT("Context"), ContextDisplayNames[SelectedContextIndex].ToString()));
-				Attributes.Add(FAnalyticsEventAttribute(TEXT("Comment"), TextBox->GetText().ToString()));
+				TArray<FAnalyticsEventAttribute> Attributes = {
+					FAnalyticsEventAttribute(TEXT("Positive"), Mode == EFeedbackMode::Positive),
+					FAnalyticsEventAttribute(TEXT("Context"), ContextDisplayNames[SelectedContextIndex].ToString()),
+					FAnalyticsEventAttribute(TEXT("Comment"), TextBox->GetText().ToString())
+				};
 
 				FEngineAnalytics::GetProvider().RecordEvent(FString("Editor.Feedback.Submitted"), Attributes);
 			}
diff --git a/Engine/Source/Programs/CrashReporter/CrashReportCommon/CrashContext.cs b/Engine/Source/Programs/CrashReporter/CrashReportCommon/CrashContext.cs
--- a/Engine/Source/Programs/CrashReporter/CrashReportCommon/CrashContext.cs
+++ b/Engine/Source/Programs/CrashReporter/CrashReportCommon/CrashContext.cs
@@ -341,6 +341,23 @@ namespace Tools.CrashReporter.CrashReportCommon
 		[XmlElement]
 		public bool IsSourceDistribution;
 
+		/// <summary></summary>
+		[XmlElement]
+		public bool IsEnsure;
+
+		/// <summary></summary>
+		[XmlElement]
+		public bool IsAssert;
+
+		/// <summary></summary>
+		[XmlElement]
+		public string CrashType;
+
+		/// <summary>
+		/// Get the crash type. Can be empty but we don't want null.
+		/// </summary>
+		public string GetCrashType() { return CrashType ?? string.Empty; }
+
 		/// <summary></summary>
 		[XmlElement]
 		public int SecondsSinceStart;
@@ -357,6 +374,10 @@ namespace Tools.CrashReporter.CrashReportCommon
 		[XmlElement]
 		public string BuildConfiguration;
 
+		/// <summary></summary>
+		[XmlElement]
+		public string GameSessionID;
+
 		/// <summary></summary>
 		[XmlElement]
 		public string PlatformName;
@@ -408,6 +429,10 @@ namespace Tools.CrashReporter.CrashReportCommon
 		[XmlElement]
 		public string EngineMode;
 
+		/// <summary></summary>
+		[XmlElement]
+		public string DeploymentName;
+
 		/// <summary></summary>
 		[XmlElement]
 		public string EngineVersion;
@@ -461,10 +486,6 @@ namespace Tools.CrashReporter.CrashReportCommon
 			return UnescapedValue.Split( new char[] { '\n' }, StringSplitOptions.RemoveEmptyEntries );
 		}
 
-		/// <summary></summary>
-		[XmlElement]
-		public string UserActivityHint;
-
 		/// <summary></summary>
 		[XmlElement]
 		public string SourceContext;
@@ -478,10 +499,6 @@ namespace Tools.CrashReporter.CrashReportCommon
 			return UnescapedValue.Split( new char[] { '\n' }, StringSplitOptions.RemoveEmptyEntries );
 		}
 
-		/// <summary> An array of module's name used by the game that crashed. </summary>
-		[XmlElement]
-		public string Modules;
-
 		/// <summary></summary>
 		[XmlElement]
 		public string UserDescription;
@@ -495,6 +512,10 @@ namespace Tools.CrashReporter.CrashReportCommon
 			return UnescapedValue.Split( new char[] { '\n' }, StringSplitOptions.RemoveEmptyEntries );
 		}
 
+		/// <summary></summary>
+		[XmlElement]
+		public string UserActivityHint;
+
 		/// <summary></summary>
 		[XmlElement]
 		public string ErrorMessage;
@@ -509,6 +530,10 @@ namespace Tools.CrashReporter.CrashReportCommon
 			return UnescapedValue.Split( new char[] { '\n' }, StringSplitOptions.RemoveEmptyEntries );
 		}
 
+		/// <summary></summary>
+		[XmlElement]
+		public int CrashDumpMode;
+
 		/// <summary> Location of the full crash dump. Displayed in the crash report frontend, without the filename. </summary>
 		[XmlElement]
 		public string FullCrashDumpLocation;
@@ -522,12 +547,6 @@ namespace Tools.CrashReporter.CrashReportCommon
 			return Fullname;
 		}
 
-		/// <summary></summary>
-		[XmlElement]
-		public int CrashDumpMode;
-
-		
-
 		/// <summary></summary>
 		[XmlElement( ElementName = "Misc.NumberOfCores" )]
 		public int Misc_NumberOfCores;
@@ -630,6 +649,18 @@ namespace Tools.CrashReporter.CrashReportCommon
 		[XmlElement]
 		public bool bAllowToBeContacted = false;
 
+		/// <summary> Platform status. </summary>
+		[XmlElement]
+		public int PlatformCallbackResult;
+
+		/// <summary> CRC version that uploaded the crash. Blank for older builds. </summary>
+		[XmlElement]
+		public string CrashReportClientVersion;
+
+		/// <summary> An array of module names used by the game that crashed. </summary>
+		[XmlElement]
+		public string Modules;
+
 		/// <summary> Warning info for missing data or failed processing in the Crash Report Processor. </summary>
 		[XmlElement]
 		public string ProcessorFailedMessage;
diff --git a/Engine/Source/Programs/CrashReporter/CrashReportCommon/ReportData.cs b/Engine/Source/Programs/CrashReporter/CrashReportCommon/ReportData.cs
--- a/Engine/Source/Programs/CrashReporter/CrashReportCommon/ReportData.cs
+++ b/Engine/Source/Programs/CrashReporter/CrashReportCommon/ReportData.cs
@@ -419,6 +419,42 @@ namespace Tools.CrashReporter.CrashReportCommon
 			}
 		}
 
+		/// <summary> Game-specific deployment string e.g. "DevPlaytest", "Live" </summary>
+		public string DeploymentName
+		{
+			get
+			{
+				return Metadata.DynamicSignatures.DeploymentName;
+			}
+		}
+
+		/// <summary> Is the report an non-fatal error? </summary>
+		public bool IsEnsure
+		{
+			get
+			{
+				return Metadata.DynamicSignatures.IsEnsure;
+			}
+		}
+
+		/// <summary> Is the report a crash from a deliberate check like an assert? </summary>
+		public bool IsAssert
+		{
+			get
+			{
+				return Metadata.DynamicSignatures.IsAssert;
+			}
+		}
+
+		/// <summary> Crash type string e.g. "Crash, "Assert", "Ensure" </summary>
+		public string CrashType
+		{
+			get
+			{
+				return Metadata.DynamicSignatures.CrashType;
+			}
+		}
+
 		/// <summary> Extracts a string marked with the specified string value. </summary>
 		private string ExtractMarkedString( string NameMarker )
 		{
diff --git a/Engine/Source/Programs/CrashReporter/CrashReportCommon/WerReport.cs b/Engine/Source/Programs/CrashReporter/CrashReportCommon/WerReport.cs
--- a/Engine/Source/Programs/CrashReporter/CrashReportCommon/WerReport.cs
+++ b/Engine/Source/Programs/CrashReporter/CrashReportCommon/WerReport.cs
@@ -93,9 +93,21 @@ namespace Tools.CrashReporter.CrashReportCommon
 		/// <summary>Generic parameter</summary>
 		public string Parameter9 = "";
 
-		/// <summary> Whether the user allowed us to be contacted. </summary>
+		/// <summary>Whether the user allowed us to be contacted</summary>
 		public bool bAllowToBeContacted = false;
 
+		/// <summary>Deployment name of the crashed app, if any</summary>
+		public string DeploymentName = "";
+
+		/// <summary>Was the "crash" a non-fatal event?</summary>
+		public bool IsEnsure = false;
+
+		/// <summary>Was the "crash" an assert or check that failed?</summary>
+		public bool IsAssert = false;
+
+		/// <summary>Crash type string</summary>
+		public string CrashType = "";
+
 		/// <summary>A simple default constructor to allow Xml serialisation.</summary>
 		public WERReportMetadataDynamicSignatures()
 		{
diff --git a/Engine/Source/Programs/CrashReporter/CrashReportProcess/Config.cs b/Engine/Source/Programs/CrashReporter/CrashReportProcess/Config.cs
--- a/Engine/Source/Programs/CrashReporter/CrashReportProcess/Config.cs
+++ b/Engine/Source/Programs/CrashReporter/CrashReportProcess/Config.cs
@@ -215,6 +215,12 @@ namespace Tools.CrashReporter.CrashReportProcess
 		[XmlElement]
 		public string AWSCredentialsFilepath { get; set; }
 
+		/// <summary>
+		/// Buffer size used to decompress zlib archives taken from S3
+		/// </summary>
+		[XmlElement]
+		public int MaxUncompressedS3RecordSize { get; set; }
+
 		/// <summary>
 		/// Index file used to store all processed crash names and times. Stops duplicates.
 		/// </summary>
diff --git a/Engine/Source/Programs/CrashReporter/CrashReportProcess/Config/CrashReportProcess.config b/Engine/Source/Programs/CrashReporter/CrashReportProcess/Config/CrashReportProcess.config
--- a/Engine/Source/Programs/CrashReporter/CrashReportProcess/Config/CrashReportProcess.config
+++ b/Engine/Source/Programs/CrashReporter/CrashReportProcess/Config/CrashReportProcess.config
@@ -33,6 +33,7 @@
 	<AWSSQSQueueUrl>https://sqs.us-west-2.amazonaws.com/queue.amazonaws.com/MyAccountNumber/MyQueueName</AWSSQSQueueUrl>
 	<AWSProfileName>default</AWSProfileName>
 	<AWSCredentialsFilepath>MyAWSCredsFilepath</AWSCredentialsFilepath>
+	<MaxUncompressedS3RecordSize>33554432</MaxUncompressedS3RecordSize>
 	<ProcessedReportsIndexPath>MyProcessedReportsIndexFilePath</ProcessedReportsIndexPath>
 	<ReportsIndexRetentionDays>120</ReportsIndexRetentionDays>
 	<AddCrashRequestTimeoutMillisec>300000</AddCrashRequestTimeoutMillisec>
diff --git a/Engine/Source/Programs/CrashReporter/CrashReportProcess/CrashReportProcessService.cs b/Engine/Source/Programs/CrashReporter/CrashReportProcess/CrashReportProcessService.cs
--- a/Engine/Source/Programs/CrashReporter/CrashReportProcess/CrashReportProcessService.cs
+++ b/Engine/Source/Programs/CrashReporter/CrashReportProcess/CrashReportProcessService.cs
@@ -89,7 +89,7 @@ namespace Tools.CrashReporter.CrashReportProcess
 			if( Message != null && Message.Length > 2 )
 			{
 				Log.Print( "[EXCEPT] " + Message );
-				StatusReporter.IncrementCount(StatusReportingConstants.ExceptionEvent);
+				StatusReporter.IncrementCount(StatusReportingEventNames.ExceptionEvent);
 			}
 		}
 
@@ -155,6 +155,8 @@ namespace Tools.CrashReporter.CrashReportProcess
 		/// </summary>
 		protected override void OnStop()
 		{
+			StatusReporter.OnPreStopping();
+
 			// Clean up the directory watcher and crash processor threads
 			foreach (var Processor in Processors)
 			{
diff --git a/Engine/Source/Programs/CrashReporter/CrashReportProcess/DataRouterConsumer.cs b/Engine/Source/Programs/CrashReporter/CrashReportProcess/DataRouterConsumer.cs
--- a/Engine/Source/Programs/CrashReporter/CrashReportProcess/DataRouterConsumer.cs
+++ b/Engine/Source/Programs/CrashReporter/CrashReportProcess/DataRouterConsumer.cs
@@ -52,16 +52,7 @@ namespace Tools.CrashReporter.CrashReportProcess
 		{
 			try
 			{
-				Int64 Size = GetVarInt(InStream);
-				Int64 MessageStart = InStream.Position;
-
-				Message = Serializer.Deserialize<ProtocolBufferRecord>(InStream);
-
-				Int64 MessageEnd = InStream.Position;
-				if (MessageEnd - MessageStart != Size)
-				{
-					throw new CrashReporterException(string.Format("DataRouterConsumer size mismatch in TryParse() - Embedded Size {0}  Actual Size {1}", Size, MessageEnd - MessageStart));
-				}
+				Message = Serializer.DeserializeWithLengthPrefix<ProtocolBufferRecord>(InStream, PrefixStyle.Base128);
 
 				return Message.HasPayload;
 			}
@@ -72,24 +63,5 @@ namespace Tools.CrashReporter.CrashReportProcess
 				return false;
 			}
 		}
-
-		private static Int64 GetVarInt(Stream InStream)
-		{
-			Int64 VarInt = 0;
-			for (int ByteIndex = 0, LeftShift = 0; ByteIndex < 9; ByteIndex++, LeftShift += 7)
-			{
-				int Byte = InStream.ReadByte();
-				if (Byte < 0)
-				{
-					throw new CrashReporterException("DataRouterConsumer encountered unexpected end of stream in GetVarInt()");
-				}
-				VarInt |= (Int64) (Byte & 0x7F) << LeftShift;
-				if ((Byte & 0x80) == 0)
-				{
-					break;
-				}
-			}
-			return VarInt;
-		}
 	}
 }
