Fix monolithic QAGameEditor CIS errors.

diff --git a/Engine/Source/Developer/HotReload/Private/HotReload.cpp b/Engine/Source/Developer/HotReload/Private/HotReload.cpp
--- a/Engine/Source/Developer/HotReload/Private/HotReload.cpp
+++ b/Engine/Source/Developer/HotReload/Private/HotReload.cpp
@@ -1564,8 +1564,10 @@ void FHotReloadModule::RecordAnalyticsEvent(const TCHAR* ReloadFrom, ECompilatio
 
 void FHotReloadModule::OnModuleCompileSucceeded(FName ModuleName, const FString& ModuleFilename)
 {
-	// If the compile succeeded, update the module info entry with the new file name for this module
+#if !IS_MONOLITHIC
+		// If the compile succeeded, update the module info entry with the new file name for this module
 	FModuleManager::Get().SetModuleFilename(ModuleName, ModuleFilename);
+#endif
 
 #if WITH_HOT_RELOAD
 	// UpdateModuleCompileData() should have been run before compiling so the
@@ -2013,12 +2015,14 @@ void FHotReloadModule::WriteModuleCompilationInfoToConfig(FName ModuleName, cons
 
 bool FHotReloadModule::GetModuleFileTimeStamp(FName ModuleName, FDateTime& OutFileTimeStamp) const
 {
+#if !IS_MONOLITHIC
 	FString Filename = FModuleManager::Get().GetModuleFilename(ModuleName);
 	if (IFileManager::Get().FileSize(*Filename) > 0)
 	{
 		OutFileTimeStamp = FDateTime(IFileManager::Get().GetTimeStamp(*Filename));
 		return true;
 	}
+#endif
 	return false;
 }
 
