Fix for warnings about engine classes being hot reloaded when they haven't even changed.

diff --git a/Engine/Plugins/Experimental/PythonScriptPlugin/Source/PythonScriptPlugin/Private/PyWrapperTypeRegistry.cpp b/Engine/Plugins/Experimental/PythonScriptPlugin/Source/PythonScriptPlugin/Private/PyWrapperTypeRegistry.cpp
--- a/Engine/Plugins/Experimental/PythonScriptPlugin/Source/PythonScriptPlugin/Private/PyWrapperTypeRegistry.cpp
+++ b/Engine/Plugins/Experimental/PythonScriptPlugin/Source/PythonScriptPlugin/Private/PyWrapperTypeRegistry.cpp
@@ -370,7 +370,8 @@ void FPyWrapperTypeReinstancer::ProcessPending()
 	{
 		for (const auto& ClassToReinstancePair : ClassesToReinstance)
 		{
-			FCoreUObjectDelegates::RegisterClassForHotReloadReinstancingDelegate.Broadcast(ClassToReinstancePair.Key, ClassToReinstancePair.Value);
+			// Assume the classes have changed
+			FCoreUObjectDelegates::RegisterClassForHotReloadReinstancingDelegate.Broadcast(ClassToReinstancePair.Key, ClassToReinstancePair.Value, EHotReloadedClassFlags::Changed);
 		}
 		FCoreUObjectDelegates::ReinstanceHotReloadedClassesDelegate.Broadcast();
 
diff --git a/Engine/Source/Developer/HotReload/Private/HotReload.cpp b/Engine/Source/Developer/HotReload/Private/HotReload.cpp
--- a/Engine/Source/Developer/HotReload/Private/HotReload.cpp
+++ b/Engine/Source/Developer/HotReload/Private/HotReload.cpp
@@ -204,7 +204,7 @@ private:
 	void ReplaceReferencesToReconstructedCDOs();
 
 #if WITH_ENGINE
-	void RegisterForReinstancing(UClass* OldClass, UClass* NewClass);
+	void RegisterForReinstancing(UClass* OldClass, UClass* NewClass, EHotReloadedClassFlags Flags);
 	void ReinstanceClasses();
 
 	/**
@@ -1219,15 +1219,23 @@ namespace {
 	}
 }
 
-void FHotReloadModule::RegisterForReinstancing(UClass* OldClass, UClass* NewClass)
+void FHotReloadModule::RegisterForReinstancing(UClass* OldClass, UClass* NewClass, EHotReloadedClassFlags Flags)
 {
 	TPair<UClass*, UClass*> Pair;
 	
 	Pair.Key = OldClass;
 	Pair.Value = NewClass;
 
-	TArray<TPair<UClass*, UClass*> >& ClassesToReinstance = GetClassesToReinstance();
-	ClassesToReinstance.Add(MoveTemp(Pair));
+	// Don't allow reinstancing of UEngine classes
+	if (!OldClass->IsChildOf(UEngine::StaticClass()) || !(Flags & EHotReloadedClassFlags::Changed))
+	{
+		TArray<TPair<UClass*, UClass*> >& ClassesToReinstance = GetClassesToReinstance();
+		ClassesToReinstance.Add(MoveTemp(Pair));
+	}
+	else
+	{
+		UE_LOG(LogHotReload, Warning, TEXT("Engine class '%s' has changed but will be ignored for hot reload"), *NewClass->GetName());
+	}
 }
 
 void FHotReloadModule::ReinstanceClasses()
@@ -1244,13 +1252,6 @@ void FHotReloadModule::ReinstanceClasses()
 	TMap<UClass*, UClass*> OldToNewClassesMap;
 	for (const TPair<UClass*, UClass*>& Pair : ClassesToReinstance)
 	{
-		// Don't allow reinstancing of UEngine classes
-		if (Pair.Key->IsChildOf(UEngine::StaticClass()))
-		{
-			UE_LOG(LogHotReload, Warning, TEXT("Engine class '%s' has changed but will be ignored for hot reload"), *Pair.Key->GetName());
-			continue;
-		}
-
 		if (Pair.Value != nullptr)
 		{
 			OldToNewClassesMap.Add(Pair.Key, Pair.Value);
@@ -1259,11 +1260,7 @@ void FHotReloadModule::ReinstanceClasses()
 
 	for (const TPair<UClass*, UClass*>& Pair : ClassesToReinstance)
 	{
-		// Don't allow reinstancing of UEngine classes
-		if (!Pair.Key->IsChildOf(UEngine::StaticClass()))
-		{
-			ReinstanceClass(Pair.Key, Pair.Value, OldToNewClassesMap);
-		}
+		ReinstanceClass(Pair.Key, Pair.Value, OldToNewClassesMap);
 	}
 
 	ClassesToReinstance.Empty();
diff --git a/Engine/Source/Runtime/CoreUObject/Private/UObject/UObjectBase.cpp b/Engine/Source/Runtime/CoreUObject/Private/UObject/UObjectBase.cpp
--- a/Engine/Source/Runtime/CoreUObject/Private/UObject/UObjectBase.cpp
+++ b/Engine/Source/Runtime/CoreUObject/Private/UObject/UObjectBase.cpp
@@ -852,7 +852,7 @@ void UClassReplaceHotReloadClasses()
 				RegisteredClass = Class->Register();
 			}
 
-			FCoreUObjectDelegates::RegisterClassForHotReloadReinstancingDelegate.Broadcast(Class->OldClass, RegisteredClass);
+			FCoreUObjectDelegates::RegisterClassForHotReloadReinstancingDelegate.Broadcast(Class->OldClass, RegisteredClass, Class->bHasChanged ? EHotReloadedClassFlags::Changed : EHotReloadedClassFlags::None);
 		}
 	}
 
diff --git a/Engine/Source/Runtime/CoreUObject/Public/UObject/UObjectGlobals.h b/Engine/Source/Runtime/CoreUObject/Public/UObject/UObjectGlobals.h
--- a/Engine/Source/Runtime/CoreUObject/Public/UObject/UObjectGlobals.h
+++ b/Engine/Source/Runtime/CoreUObject/Public/UObject/UObjectGlobals.h
@@ -1958,6 +1958,16 @@ class FPackageReloadedEvent;
 
 class FGarbageCollectionTracer;
 
+enum class EHotReloadedClassFlags
+{
+	None = 0,
+
+	// Set when the hot reloaded class has been detected as changed
+	Changed = 0x01
+};
+
+ENUM_CLASS_FLAGS(EHotReloadedClassFlags)
+
 /**
  * Global CoreUObject delegates
  */
@@ -2023,7 +2033,7 @@ struct COREUOBJECT_API FCoreUObjectDelegates
 	static FRegisterHotReloadAddedClassesDelegate RegisterHotReloadAddedClassesDelegate;
 
 	/** Delegate for registering hot-reloaded classes that changed after hot-reload for reinstancing */
-	DECLARE_MULTICAST_DELEGATE_TwoParams(FRegisterClassForHotReloadReinstancingDelegate, UClass*, UClass*);
+	DECLARE_MULTICAST_DELEGATE_ThreeParams(FRegisterClassForHotReloadReinstancingDelegate, UClass*, UClass*, EHotReloadedClassFlags);
 	static FRegisterClassForHotReloadReinstancingDelegate RegisterClassForHotReloadReinstancingDelegate;
 
 	/** Delegate for reinstancing hot-reloaded classes */
