Fix hot reload compiling all game modules. Caused by AnimationSharing module declaring itself as a game module (via IModuleInterface), and hot reload disabling shared PCHs if a non-game module (determined via directory location) is being rebuilt.

(reimplementing CL #4755564 from Dev-Core)

diff --git a/Engine/Plugins/Developer/AnimationSharing/Source/AnimationSharing/Public/AnimationSharingModule.h b/Engine/Plugins/Developer/AnimationSharing/Source/AnimationSharing/Public/AnimationSharingModule.h
--- a/Engine/Plugins/Developer/AnimationSharing/Source/AnimationSharing/Public/AnimationSharingModule.h
+++ b/Engine/Plugins/Developer/AnimationSharing/Source/AnimationSharing/Public/AnimationSharingModule.h
@@ -12,7 +12,7 @@
 class UAnimationSharingManager;
 class UAnimationSharingSetup;
 
-class ANIMATIONSHARING_API FAnimSharingModule : public FDefaultGameModuleImpl, public FGCObject
+class ANIMATIONSHARING_API FAnimSharingModule : public FDefaultModuleImpl, public FGCObject
 {
 public:
 	// Begin IModuleInterface overrides
diff --git a/Engine/Source/Developer/HotReload/Private/HotReload.cpp b/Engine/Source/Developer/HotReload/Private/HotReload.cpp
--- a/Engine/Source/Developer/HotReload/Private/HotReload.cpp
+++ b/Engine/Source/Developer/HotReload/Private/HotReload.cpp
@@ -1632,12 +1632,6 @@ bool FHotReloadModule::StartCompilingModuleDLLs(const TArray< FModuleToRecompile
 		ExtraArg += TEXT( "-FailIfGeneratedCodeChanges " );
 	}
 
-	// Shared PCH does no work with hot-reloading engine/editor modules as we don't scan all modules for them.
-	if (!ContainsOnlyGameModules(ModuleNames))
-	{
-		ExtraArg += TEXT("-nosharedpch ");
-	}
-	
 	FString CmdLineParams = FString::Printf( TEXT( "%s %s %s %s%s -IgnoreJunk" ), 
 		*ModuleArg, 
 		BuildPlatformName, BuildConfigurationName, 
