Reset the module manager cache after recompiling an unloaded module just in case it's a brand new one.

diff --git a/Engine/Source/Developer/HotReload/Private/HotReload.cpp b/Engine/Source/Developer/HotReload/Private/HotReload.cpp
--- a/Engine/Source/Developer/HotReload/Private/HotReload.cpp
+++ b/Engine/Source/Developer/HotReload/Private/HotReload.cpp
@@ -675,6 +675,11 @@ bool FHotReloadModule::RecompileModule(const FName InModuleName, const bool bRel
 		Ar.Logf( TEXT( "Unloading module before compile." ) );
 		ModuleManager.UnloadOrAbandonModuleWithCallback( InModuleName, Ar );
 	}
+	else
+	{
+		// Reset the module cache in case it's a new module that we probably didn't know about already.
+		ModuleManager.ResetModulePathsCache();
+	}
 
 	// Reload the module if it was loaded before we recompiled
 	if ((bWasModuleLoaded || bForceCodeProject) && bReloadAfterRecompile)
