Fix for reinstancing of engine classes in hot reload.

diff --git a/Engine/Source/Developer/HotReload/Private/HotReload.cpp b/Engine/Source/Developer/HotReload/Private/HotReload.cpp
--- a/Engine/Source/Developer/HotReload/Private/HotReload.cpp
+++ b/Engine/Source/Developer/HotReload/Private/HotReload.cpp
@@ -1227,12 +1227,12 @@ void FHotReloadModule::RegisterForReinstancing(UClass* OldClass, UClass* NewClas
 	Pair.Value = NewClass;
 
 	// Don't allow reinstancing of UEngine classes
-	if (!OldClass->IsChildOf(UEngine::StaticClass()) || !(Flags & EHotReloadedClassFlags::Changed))
+	if (!OldClass->IsChildOf(UEngine::StaticClass()))
 	{
 		TArray<TPair<UClass*, UClass*> >& ClassesToReinstance = GetClassesToReinstance();
 		ClassesToReinstance.Add(MoveTemp(Pair));
 	}
-	else
+	else if (EnumHasAnyFlags(Flags, EHotReloadedClassFlags::Changed))
 	{
 		UE_LOG(LogHotReload, Warning, TEXT("Engine class '%s' has changed but will be ignored for hot reload"), *NewClass->GetName());
 	}
