Removed use of deprecated NAME_INDEX typedef

diff --git a/Engine/Source/Developer/HotReload/Private/HotReloadClassReinstancer.cpp b/Engine/Source/Developer/HotReload/Private/HotReloadClassReinstancer.cpp
--- a/Engine/Source/Developer/HotReload/Private/HotReloadClassReinstancer.cpp
+++ b/Engine/Source/Developer/HotReload/Private/HotReloadClassReinstancer.cpp
@@ -146,8 +146,8 @@ void FHotReloadClassReinstancer::SerializeCDOProperties(UObject* InObject, FHotR
 		virtual FArchive& operator<<(FName& InName) override
 		{
 			FArchive& Ar = *this;
-			NAME_INDEX ComparisonIndex = InName.GetComparisonIndex();
-			NAME_INDEX DisplayIndex = InName.GetDisplayIndex();
+			FNameEntryId ComparisonIndex = InName.GetComparisonIndex();
+			FNameEntryId DisplayIndex = InName.GetDisplayIndex();
 			int32 Number = InName.GetNumber();
 			Ar << ComparisonIndex;
 			Ar << DisplayIndex;
@@ -385,8 +385,8 @@ void FHotReloadClassReinstancer::UpdateDefaultProperties()
 		virtual FArchive& operator<<(FName& InName) override
 		{
 			FArchive& Ar = *this;
-			NAME_INDEX ComparisonIndex = InName.GetComparisonIndex();
-			NAME_INDEX DisplayIndex = InName.GetDisplayIndex();
+			FNameEntryId ComparisonIndex = InName.GetComparisonIndex();
+			FNameEntryId DisplayIndex = InName.GetDisplayIndex();
 			int32 Number = InName.GetNumber();
 			Ar << ComparisonIndex;
 			Ar << DisplayIndex;
diff --git a/Engine/Source/Runtime/AIModule/Classes/Perception/AISense_Blueprint.h b/Engine/Source/Runtime/AIModule/Classes/Perception/AISense_Blueprint.h
--- a/Engine/Source/Runtime/AIModule/Classes/Perception/AISense_Blueprint.h
+++ b/Engine/Source/Runtime/AIModule/Classes/Perception/AISense_Blueprint.h
@@ -79,5 +79,5 @@ protected:
 	void OnListenerRemovedImpl(const FPerceptionListener& UpdatedListener);
 
 private:
-	static TMap<NAME_INDEX, FAISenseID> BPSenseToSenseID;
+	static TMap<FNameEntryId, FAISenseID> BPSenseToSenseID;
 };
diff --git a/Engine/Source/Runtime/AIModule/Private/Perception/AISense_Blueprint.cpp b/Engine/Source/Runtime/AIModule/Private/Perception/AISense_Blueprint.cpp
--- a/Engine/Source/Runtime/AIModule/Private/Perception/AISense_Blueprint.cpp
+++ b/Engine/Source/Runtime/AIModule/Private/Perception/AISense_Blueprint.cpp
@@ -10,7 +10,7 @@
 //----------------------------------------------------------------------//
 // 
 //----------------------------------------------------------------------//
-TMap<NAME_INDEX, FAISenseID> UAISense_Blueprint::BPSenseToSenseID;
+TMap<FNameEntryId, FAISenseID> UAISense_Blueprint::BPSenseToSenseID;
 
 UAISense_Blueprint::UAISense_Blueprint(const FObjectInitializer& ObjectInitializer)
 	: Super(ObjectInitializer)
@@ -92,7 +92,7 @@ FAISenseID UAISense_Blueprint::UpdateSenseID()
 #endif
 	if (GetClass()->HasAnyClassFlags(CLASS_Abstract) == false)
 	{
-		const NAME_INDEX NameIndex = GetClass()->GetFName().GetDisplayIndex();
+		const FNameEntryId NameIndex = GetClass()->GetFName().GetDisplayIndex();
 		const FAISenseID* StoredID = BPSenseToSenseID.Find(NameIndex);
 		if (StoredID != nullptr)
 		{
diff --git a/Engine/Source/Runtime/AssetRegistry/Public/AssetDataTagMap.h b/Engine/Source/Runtime/AssetRegistry/Public/AssetDataTagMap.h
--- a/Engine/Source/Runtime/AssetRegistry/Public/AssetDataTagMap.h
+++ b/Engine/Source/Runtime/AssetRegistry/Public/AssetDataTagMap.h
@@ -108,9 +108,9 @@ private:
 	};
 	struct FCompactExportTextNoNumbers
 	{
-		NAME_INDEX Class;
-		NAME_INDEX Package;
-		NAME_INDEX Object;
+		FNameEntryId Class;
+		FNameEntryId Package;
+		FNameEntryId Object;
 
 		FString ToString() const
 		{
diff --git a/Engine/Source/Runtime/CoreUObject/Private/Serialization/DuplicateDataReader.cpp b/Engine/Source/Runtime/CoreUObject/Private/Serialization/DuplicateDataReader.cpp
--- a/Engine/Source/Runtime/CoreUObject/Private/Serialization/DuplicateDataReader.cpp
+++ b/Engine/Source/Runtime/CoreUObject/Private/Serialization/DuplicateDataReader.cpp
@@ -47,8 +47,8 @@ void FDuplicateDataReader::SerializeFail()
 
 FArchive& FDuplicateDataReader::operator<<(FName& N)
 {
-	NAME_INDEX ComparisonIndex;
-	NAME_INDEX DisplayIndex;
+	FNameEntryId ComparisonIndex;
+	FNameEntryId DisplayIndex;
 	int32 Number;
 	ByteOrderSerialize(&ComparisonIndex, sizeof(ComparisonIndex));
 	ByteOrderSerialize(&DisplayIndex, sizeof(DisplayIndex));
diff --git a/Engine/Source/Runtime/CoreUObject/Private/Serialization/DuplicateDataWriter.cpp b/Engine/Source/Runtime/CoreUObject/Private/Serialization/DuplicateDataWriter.cpp
--- a/Engine/Source/Runtime/CoreUObject/Private/Serialization/DuplicateDataWriter.cpp
+++ b/Engine/Source/Runtime/CoreUObject/Private/Serialization/DuplicateDataWriter.cpp
@@ -45,8 +45,8 @@ FDuplicateDataWriter::FDuplicateDataWriter( FUObjectAnnotationSparse<FDuplicated
  */
 FArchive& FDuplicateDataWriter::operator<<(FName& N)
 {
-	NAME_INDEX ComparisonIndex = N.GetComparisonIndex();
-	NAME_INDEX DisplayIndex = N.GetDisplayIndex();
+	FNameEntryId ComparisonIndex = N.GetComparisonIndex();
+	FNameEntryId DisplayIndex = N.GetDisplayIndex();
 	int32 Number = N.GetNumber();
 	ByteOrderSerialize(&ComparisonIndex, sizeof(ComparisonIndex));
 	ByteOrderSerialize(&DisplayIndex, sizeof(DisplayIndex));
diff --git a/Engine/Source/Runtime/CoreUObject/Private/Serialization/ObjectReader.cpp b/Engine/Source/Runtime/CoreUObject/Private/Serialization/ObjectReader.cpp
--- a/Engine/Source/Runtime/CoreUObject/Private/Serialization/ObjectReader.cpp
+++ b/Engine/Source/Runtime/CoreUObject/Private/Serialization/ObjectReader.cpp
@@ -9,8 +9,8 @@
 
 FArchive& FObjectReader::operator<<(FName& N)
 {
-	NAME_INDEX ComparisonIndex;
-	NAME_INDEX DisplayIndex;
+	FNameEntryId ComparisonIndex;
+	FNameEntryId DisplayIndex;
 	int32 Number;
 	ByteOrderSerialize(&ComparisonIndex, sizeof(ComparisonIndex));
 	ByteOrderSerialize(&DisplayIndex, sizeof(DisplayIndex));
diff --git a/Engine/Source/Runtime/CoreUObject/Private/Serialization/ObjectWriter.cpp b/Engine/Source/Runtime/CoreUObject/Private/Serialization/ObjectWriter.cpp
--- a/Engine/Source/Runtime/CoreUObject/Private/Serialization/ObjectWriter.cpp
+++ b/Engine/Source/Runtime/CoreUObject/Private/Serialization/ObjectWriter.cpp
@@ -9,8 +9,8 @@
 
 FArchive& FObjectWriter::operator<<( class FName& N )
 {
-	NAME_INDEX ComparisonIndex = N.GetComparisonIndex();
-	NAME_INDEX DisplayIndex = N.GetDisplayIndex();
+	FNameEntryId ComparisonIndex = N.GetComparisonIndex();
+	FNameEntryId DisplayIndex = N.GetDisplayIndex();
 	int32 Number = N.GetNumber();
 	ByteOrderSerialize(&ComparisonIndex, sizeof(ComparisonIndex));
 	ByteOrderSerialize(&DisplayIndex, sizeof(DisplayIndex));
diff --git a/Engine/Source/Runtime/Engine/Private/Materials/MaterialShared.cpp b/Engine/Source/Runtime/Engine/Private/Materials/MaterialShared.cpp
--- a/Engine/Source/Runtime/Engine/Private/Materials/MaterialShared.cpp
+++ b/Engine/Source/Runtime/Engine/Private/Materials/MaterialShared.cpp
@@ -3805,7 +3805,7 @@ void FMaterialResourceMemoryWriter::SerializeToParentArchive()
 	check(Ar.IsSaving() && this->IsByteSwapping() == Ar.IsByteSwapping());
 
 	// Make a array of unique names used by the shader map
-	TArray<NAME_INDEX> DisplayIndices;
+	TArray<FNameEntryId> DisplayIndices;
 	auto NumNames = Name2Indices.Num();
 	DisplayIndices.Empty(NumNames);
 	DisplayIndices.AddDefaulted(NumNames);
@@ -3815,7 +3815,7 @@ void FMaterialResourceMemoryWriter::SerializeToParentArchive()
 	}
 
 	Ar << NumNames;
-	for (NAME_INDEX DisplayIdx : DisplayIndices)
+	for (FNameEntryId DisplayIdx : DisplayIndices)
 	{
 		FName::GetEntry(DisplayIdx)->Write(Ar);
 	}
diff --git a/Engine/Source/Runtime/Engine/Public/MaterialShared.h b/Engine/Source/Runtime/Engine/Public/MaterialShared.h
--- a/Engine/Source/Runtime/Engine/Public/MaterialShared.h
+++ b/Engine/Source/Runtime/Engine/Public/MaterialShared.h
@@ -2634,7 +2634,7 @@ public:
 private:
 	TArray<uint8> Bytes;
 	TArray<FMaterialResourceLocOnDisk> Locs;
-	TMap<NAME_INDEX, int32> Name2Indices;
+	TMap<FNameEntryId, int32> Name2Indices;
 	FArchive* ParentAr;
 
 	void SerializeToParentArchive();
