Fix for hot reload crashing caused by reinstancing redundant BP objects whose class has been changed but still have the old layout.

diff --git a/Engine/Source/Developer/HotReload/Private/HotReloadClassReinstancer.cpp b/Engine/Source/Developer/HotReload/Private/HotReloadClassReinstancer.cpp
--- a/Engine/Source/Developer/HotReload/Private/HotReloadClassReinstancer.cpp
+++ b/Engine/Source/Developer/HotReload/Private/HotReloadClassReinstancer.cpp
@@ -481,9 +481,15 @@ void FHotReloadClassReinstancer::UpdateDefaultProperties()
 		TArray<uint8> CurrentValueSerializedData;		
 
 		// Update properties on all existing instances of the class
+		const UPackage* TransientPackage = GetTransientPackage();
 		for (FObjectIterator It(NewClass); It; ++It)
 		{
 			UObject* ObjectPtr = *It;
+			if (ObjectPtr->IsPendingKill() || ObjectPtr->GetOutermost() == TransientPackage)
+			{
+				continue;
+			}
+
 			DefaultSubobjectArray.Empty(DefaultSubobjectArrayCapacity);
 			ObjectPtr->CollectDefaultSubobjects(DefaultSubobjectArray);
 
