Replacement of some hot reload's bool flags with enum class flags.

diff --git a/Engine/Source/Developer/HotReload/Private/HotReload.cpp b/Engine/Source/Developer/HotReload/Private/HotReload.cpp
--- a/Engine/Source/Developer/HotReload/Private/HotReload.cpp
+++ b/Engine/Source/Developer/HotReload/Private/HotReload.cpp
@@ -108,7 +108,7 @@ public:
 
 	/** IHotReloadInterface implementation */
 	virtual void SaveConfig() override;
-	virtual bool RecompileModule(const FName InModuleName, const bool bReloadAfterRecompile, FOutputDevice &Ar, bool bFailIfGeneratedCodeChanges = true, bool bForceCodeProject = false) override;
+	virtual bool RecompileModule(const FName InModuleName, FOutputDevice &Ar, ERecompileModuleFlags Flags) override;
 	virtual bool IsCurrentlyCompiling() const override { return ModuleCompileProcessHandle.IsValid(); }
 	virtual void RequestStopCompilation() override { bRequestCancelCompilation = true; }
 	virtual void AddHotReloadFunctionRemap(FNativeFuncPtr NewFunctionPointer, FNativeFuncPtr OldFunctionPointer) override;	
@@ -255,14 +255,13 @@ private:
 	 *	@param ModuleNames The list of modules to compile.
 	 *	@param InRecompileModulesCallback Callback function to make when module recompiles.
 	 *	@param Ar
-	 *	@param bInFailIfGeneratedCodeChanges If true, fail the compilation if generated headers change.
 	 *	@param InAdditionalCmdLineArgs Additional arguments to pass to UBT.
-	 *  @param bForceCodeProject Compile as code-based project even if there's no game modules loaded
+	 *  @param Flags Compilation flags
 	 *	@return true if successful, false otherwise.
 	 */
-	bool StartCompilingModuleDLLs(const TArray< FModuleToRecompile >& ModuleNames, 
-		FRecompileModulesCallback&& InRecompileModulesCallback, FOutputDevice& Ar, bool bInFailIfGeneratedCodeChanges, 
-		const FString& InAdditionalCmdLineArgs, bool bForceCodeProject);
+	bool StartCompilingModuleDLLs(const TArray< FModuleToRecompile >& ModuleNames,
+		FRecompileModulesCallback&& InRecompileModulesCallback, FOutputDevice& Ar,
+		const FString& InAdditionalCmdLineArgs, ERecompileModuleFlags Flags);
 #endif
 
 	/** Launches UnrealBuildTool with the specified command line parameters */
@@ -557,10 +556,7 @@ bool FHotReloadModule::Exec( UWorld* Inworld, const TCHAR* Cmd, FOutputDevice& A
 			if( !ModuleNameStr.IsEmpty() )
 			{
 				const FName ModuleName( *ModuleNameStr );
-				const bool bReloadAfterRecompile = true;
-				const bool bForceCodeProject = false;
-				const bool bFailIfGeneratedCodeChanges = true;
-				RecompileModule( ModuleName, bReloadAfterRecompile, Ar, bFailIfGeneratedCodeChanges, bForceCodeProject);
+				RecompileModule(ModuleName, Ar, ERecompileModuleFlags::ReloadAfterRecompile | ERecompileModuleFlags::FailIfGeneratedCodeChanges);
 			}
 
 			return true;
@@ -602,7 +598,7 @@ FString FHotReloadModule::GetModuleCompileMethod(FName InModuleName)
 	}
 }
 
-bool FHotReloadModule::RecompileModule(const FName InModuleName, const bool bReloadAfterRecompile, FOutputDevice &Ar, bool bFailIfGeneratedCodeChanges, bool bForceCodeProject)
+bool FHotReloadModule::RecompileModule(const FName InModuleName, FOutputDevice &Ar, ERecompileModuleFlags Flags)
 {
 #if WITH_HOT_RELOAD
 
@@ -653,11 +649,11 @@ bool FHotReloadModule::RecompileModule(const FName InModuleName, const bool bRel
 	 * @param Ar Output device for logging compilation status.
 	 * @param bForceCodeProject Even if it's a non-code project, treat it as code-based project
 	 */
-	auto RecompileModuleDLLs = [this, &Ar, bFailIfGeneratedCodeChanges, bForceCodeProject](const TArray< FModuleToRecompile >& ModuleNames)
+	auto RecompileModuleDLLs = [this, &Ar, Flags](const TArray< FModuleToRecompile >& ModuleNames)
 	{
 		bool bCompileSucceeded = false;
 		const FString AdditionalArguments = MakeUBTArgumentsForModuleCompiling();
-		if (StartCompilingModuleDLLs(ModuleNames, nullptr, Ar, bFailIfGeneratedCodeChanges, AdditionalArguments, bForceCodeProject))
+		if (StartCompilingModuleDLLs(ModuleNames, nullptr, Ar, AdditionalArguments, Flags))
 		{
 			bool bCompileStillInProgress = false;
 			CheckForFinishedModuleDLLCompile( EHotReloadFlags::WaitForCompletion, bCompileStillInProgress, bCompileSucceeded, Ar );
@@ -696,7 +692,7 @@ bool FHotReloadModule::RecompileModule(const FName InModuleName, const bool bRel
 	}
 
 	// Reload the module if it was loaded before we recompiled
-	if ((bWasModuleLoaded || bForceCodeProject) && bReloadAfterRecompile)
+	if ((bWasModuleLoaded || !!(Flags & ERecompileModuleFlags::ForceCodeProject)) && !!(Flags & ERecompileModuleFlags::ReloadAfterRecompile))
 	{
 		TGuardValue<bool> GuardIsHotReload(GIsHotReload, true);
 		Ar.Logf( TEXT( "Reloading module %s after successful compile." ), *InModuleName.ToString() );
@@ -708,7 +704,7 @@ bool FHotReloadModule::RecompileModule(const FName InModuleName, const bool bRel
 		CollectGarbage(GARBAGE_COLLECTION_KEEPFLAGS);
 	}
 
-	if (bForceCodeProject)
+	if (!!(Flags & ERecompileModuleFlags::ForceCodeProject))
 	{
 		HotReloadEvent.Broadcast( false );
 	}
@@ -1165,9 +1161,8 @@ ECompilationResult::Type FHotReloadModule::RebindPackagesInternal(const TArray<U
 			DoHotReloadInternal(ChangedModules, InPackages, DependentModules, Ar);
 		},
 		Ar,
-		false, /* bFailIfGeneratedCodeChanges */
 		AdditionalArguments,
-		false /* bForceCodeProject */
+		ERecompileModuleFlags::None
 	);
 
 	if (!bCompileStarted)
@@ -1602,9 +1597,9 @@ FString FHotReloadModule::MakeUBTArgumentsForModuleCompiling()
 }
 
 #if WITH_HOT_RELOAD
-bool FHotReloadModule::StartCompilingModuleDLLs(const TArray< FModuleToRecompile >& ModuleNames, 
-	FRecompileModulesCallback&& InRecompileModulesCallback, FOutputDevice& Ar, bool bInFailIfGeneratedCodeChanges, 
-	const FString& InAdditionalCmdLineArgs, bool bForceCodeProject)
+bool FHotReloadModule::StartCompilingModuleDLLs(const TArray< FModuleToRecompile >& ModuleNames,
+	FRecompileModulesCallback&& InRecompileModulesCallback, FOutputDevice& Ar,
+	const FString& InAdditionalCmdLineArgs, ERecompileModuleFlags Flags)
 {
 	// Keep track of what we're compiling
 	ModulesBeingCompiled.Empty(ModuleNames.Num());
@@ -1651,7 +1646,7 @@ bool FHotReloadModule::StartCompilingModuleDLLs(const TArray< FModuleToRecompile
 		ExtraArg += FString::Printf(TEXT("-Project=\"%s\" "), *FPaths::ConvertRelativePathToFull(FPaths::GetProjectFilePath()));
 	}
 
-	if (bInFailIfGeneratedCodeChanges)
+	if (!!(Flags & ERecompileModuleFlags::FailIfGeneratedCodeChanges))
 	{
 		// Additional argument to let UHT know that we can only compile the module if the generated code didn't change
 		ExtraArg += TEXT( "-FailIfGeneratedCodeChanges " );
diff --git a/Engine/Source/Developer/ModuleUI/Private/SModuleUI.cpp b/Engine/Source/Developer/ModuleUI/Private/SModuleUI.cpp
--- a/Engine/Source/Developer/ModuleUI/Private/SModuleUI.cpp
+++ b/Engine/Source/Developer/ModuleUI/Private/SModuleUI.cpp
@@ -289,10 +289,7 @@ FReply SModuleUI::FModuleListItem::OnRecompileClicked()
 		else
 		{
 			// Perform a regular unload, then reload
-			const bool bReloadAfterRecompile = true;
-			const bool bForceCodeProject = false;
-			const bool bFailIfGeneratedCodeChanges = true;
-			const bool bRecompileSucceeded = HotReloadSupport.RecompileModule(ModuleName, bReloadAfterRecompile, *GLog, bFailIfGeneratedCodeChanges, bForceCodeProject);
+			const bool bRecompileSucceeded = HotReloadSupport.RecompileModule(ModuleName, *GLog, ERecompileModuleFlags::ReloadAfterRecompile | ERecompileModuleFlags::FailIfGeneratedCodeChanges);
 		}
 	}
 
diff --git a/Engine/Source/Editor/GameProjectGeneration/Private/GameProjectUtils.cpp b/Engine/Source/Editor/GameProjectGeneration/Private/GameProjectUtils.cpp
--- a/Engine/Source/Editor/GameProjectGeneration/Private/GameProjectUtils.cpp
+++ b/Engine/Source/Editor/GameProjectGeneration/Private/GameProjectUtils.cpp
@@ -3960,10 +3960,7 @@ GameProjectUtils::EAddCodeToProjectResult GameProjectUtils::AddCodeToProject_Int
 		check(ModuleInfo.ModuleName == GameModuleName);
 
 		IHotReloadInterface& HotReloadSupport = FModuleManager::LoadModuleChecked<IHotReloadInterface>("HotReload");
-		const bool bReloadAfterCompiling = true;
-		const bool bForceCodeProject = true;
-		const bool bFailIfGeneratedCodeChanges = false;
-		if (!HotReloadSupport.RecompileModule(*GameModuleName, bReloadAfterCompiling, *GWarn, bFailIfGeneratedCodeChanges, bForceCodeProject))
+		if (!HotReloadSupport.RecompileModule(*GameModuleName, *GWarn, ERecompileModuleFlags::ReloadAfterRecompile | ERecompileModuleFlags::ForceCodeProject))
 		{
 			OutFailReason = LOCTEXT("FailedToCompileNewGameModule", "Failed to compile newly created game module.");
 			return EAddCodeToProjectResult::FailedToHotReload;
@@ -4009,10 +4006,7 @@ GameProjectUtils::EAddCodeToProjectResult GameProjectUtils::AddCodeToProject_Int
 			else
 			{
 				// Perform a regular unload, then reload
-				const bool bReloadAfterRecompile = true;
-				const bool bForceCodeProject = false;
-				const bool bFailIfGeneratedCodeChanges = true;
-				if (!HotReloadSupport.RecompileModule(ModuleFName, bReloadAfterRecompile, *GWarn, bFailIfGeneratedCodeChanges, bForceCodeProject))
+				if (!HotReloadSupport.RecompileModule(ModuleFName, *GWarn, ERecompileModuleFlags::ReloadAfterRecompile | ERecompileModuleFlags::FailIfGeneratedCodeChanges))
 				{
 					OutFailReason = FText::Format(LOCTEXT("FailedToCompileModuleFmt", "Failed to automatically compile the '{0}' module."), FText::FromString(ModuleInfo.ModuleName));
 					return EAddCodeToProjectResult::FailedToHotReload;
diff --git a/Engine/Source/Editor/UnrealEd/Private/Factories/PackFactory.cpp b/Engine/Source/Editor/UnrealEd/Private/Factories/PackFactory.cpp
--- a/Engine/Source/Editor/UnrealEd/Private/Factories/PackFactory.cpp
+++ b/Engine/Source/Editor/UnrealEd/Private/Factories/PackFactory.cpp
@@ -697,10 +697,7 @@ UObject* UPackFactory::FactoryCreateBinary
 					}
 					else
 					{
-						const bool bReloadAfterCompiling = true;
-						const bool bForceCodeProject = true;
-						const bool bFailIfGeneratedCodeChanges = false;
-						if (!HotReloadSupport.RecompileModule(FApp::GetProjectName(), bReloadAfterCompiling, *GWarn, bFailIfGeneratedCodeChanges, bForceCodeProject))
+						if (!HotReloadSupport.RecompileModule(FApp::GetProjectName(), *GWarn, ERecompileModuleFlags::ReloadAfterRecompile | ERecompileModuleFlags::ForceCodeProject))
 						{
 							FMessageDialog::Open(EAppMsgType::Ok, NSLOCTEXT("PackFactory", "FailedToCompileNewGameModule", "Failed to compile newly created game module."));
 						}
diff --git a/Engine/Source/Runtime/CoreUObject/Public/Misc/HotReloadInterface.h b/Engine/Source/Runtime/CoreUObject/Public/Misc/HotReloadInterface.h
--- a/Engine/Source/Runtime/CoreUObject/Public/Misc/HotReloadInterface.h
+++ b/Engine/Source/Runtime/CoreUObject/Public/Misc/HotReloadInterface.h
@@ -15,9 +15,23 @@ enum class EHotReloadFlags : uint8
 	// Should not return until the recompile and reload has completed
 	WaitForCompletion = 0x01
 };
-
 ENUM_CLASS_FLAGS(EHotReloadFlags)
 
+enum class ERecompileModuleFlags : uint8
+{
+	None = 0x00,
+
+	// Perform a reload of the module after the recompile finishes
+	ReloadAfterRecompile = 0x01,
+
+	// Report failure if UHT-generated code changes as a result of the recompile
+	FailIfGeneratedCodeChanges = 0x02,
+
+	// Even if this is not code-based project compile with game project as the target for UBT (do not use UE4Editor target)
+	ForceCodeProject = 0x04,
+};
+ENUM_CLASS_FLAGS(ERecompileModuleFlags)
+
 /**
  * HotReload module interface
  */
@@ -50,11 +64,10 @@ public:
 	 * Recompiles a single module
 	 *
 	 * @param InModuleName Name of the module to compile
-	 * @param bReloadAfterRecompile Should the module be reloaded after recompile
 	 * @param Ar Output device (logging)
-	 * @param bForceCodeProject Even if this is not code-based project compile with game project as the target for UBT (do not use UE4Editor target)
+	 * @param Flags Recompilation flags
 	 */
-	virtual bool RecompileModule(const FName InModuleName, const bool bReloadAfterRecompile, FOutputDevice &Ar, bool bFailIfGeneratedCodeChanges = true, bool bForceCodeProject = false) = 0;
+	virtual bool RecompileModule(const FName InModuleName, FOutputDevice &Ar, ERecompileModuleFlags Flags) = 0;
 
 	/**
 	 * Returns whether modules are currently being compiled
diff --git a/Engine/Source/Runtime/Engine/Private/RendererSupport.cpp b/Engine/Source/Runtime/Engine/Private/RendererSupport.cpp
--- a/Engine/Source/Runtime/Engine/Private/RendererSupport.cpp
+++ b/Engine/Source/Runtime/Engine/Private/RendererSupport.cpp
@@ -121,9 +121,7 @@ void RecompileRendererModule()
 		bool bCompiledSuccessfully = false;
 		do 
 		{
-			const bool bForceCodeProject = false;
-			const bool bFailIfGeneratedCodeChanges = true;
-			bCompiledSuccessfully = HotReload->RecompileModule(RendererModuleName, false, *GLog, bFailIfGeneratedCodeChanges, bForceCodeProject);
+			bCompiledSuccessfully = HotReload->RecompileModule(RendererModuleName, *GLog, ERecompileModuleFlags::FailIfGeneratedCodeChanges);
 
 			if (!bCompiledSuccessfully)
 			{
