Fixed some cases where the editor tries to spawn UBT to recompile modules using -TargetType=?. For projects with multiple editor targets, this won't work because UBT won't know which target to compile.

diff --git a/Engine/Plugins/Editor/PluginBrowser/Source/PluginBrowser/Private/SNewPluginWizard.cpp b/Engine/Plugins/Editor/PluginBrowser/Source/PluginBrowser/Private/SNewPluginWizard.cpp
--- a/Engine/Plugins/Editor/PluginBrowser/Source/PluginBrowser/Private/SNewPluginWizard.cpp
+++ b/Engine/Plugins/Editor/PluginBrowser/Source/PluginBrowser/Private/SNewPluginWizard.cpp
@@ -743,7 +743,7 @@ FReply SNewPluginWizard::OnCreatePluginClicked()
 	if (bSucceeded && bHasModules)
 	{
 		FString ProjectFileName = IFileManager::Get().ConvertToAbsolutePathForExternalAppForWrite(*FPaths::GetProjectFilePath());
-		FString Arguments = FString::Printf(TEXT("%s %s -TargetType=Editor -Plugin=\"%s\" -Project=\"%s\" -Progress -NoHotReloadFromIDE"), FModuleManager::Get().GetUBTConfiguration(), FPlatformMisc::GetUBTPlatform(), *UPluginFilePath, *ProjectFileName);
+		FString Arguments = FString::Printf(TEXT("%s %s %s -Plugin=\"%s\" -Project=\"%s\" -Progress -NoHotReloadFromIDE"), FPlatformMisc::GetUBTTargetName(), FModuleManager::Get().GetUBTConfiguration(), FPlatformMisc::GetUBTPlatform(), *UPluginFilePath, *ProjectFileName);
 		if (!FDesktopPlatformModule::Get()->RunUnrealBuildTool(LOCTEXT("Compiling", "Compiling..."), FPaths::RootDir(), Arguments, GWarn))
 		{
 			PopErrorNotification(LOCTEXT("FailedToCompile", "Failed to compile source code."));
diff --git a/Engine/Source/Developer/HotReload/Private/HotReload.cpp b/Engine/Source/Developer/HotReload/Private/HotReload.cpp
--- a/Engine/Source/Developer/HotReload/Private/HotReload.cpp
+++ b/Engine/Source/Developer/HotReload/Private/HotReload.cpp
@@ -1608,6 +1608,7 @@ bool FHotReloadModule::StartCompilingModuleDLLs(const TArray< FModuleToRecompile
 
 	const TCHAR* BuildPlatformName = FPlatformMisc::GetUBTPlatform();
 	const TCHAR* BuildConfigurationName = FModuleManager::GetUBTConfiguration();
+	const TCHAR* BuildTargetName = FPlatformMisc::GetUBTTargetName();
 
 	RecompileModulesCallback = MoveTemp(InRecompileModulesCallback);
 
@@ -1636,10 +1637,6 @@ bool FHotReloadModule::StartCompilingModuleDLLs(const TArray< FModuleToRecompile
 	}
 
 	FString ExtraArg;
-#if UE_EDITOR
-	// When recompiling from the editor, we don't know the editor target name. Pass -TargetType=Editor to UBT to have it figure it out.
-	ExtraArg = TEXT( "-TargetType=Editor " );
-#endif
 
 	if (FPaths::IsProjectFilePathSet())
 	{
@@ -1652,9 +1649,9 @@ bool FHotReloadModule::StartCompilingModuleDLLs(const TArray< FModuleToRecompile
 		ExtraArg += TEXT( "-FailIfGeneratedCodeChanges " );
 	}
 
-	FString CmdLineParams = FString::Printf( TEXT( "%s %s %s %s%s -IgnoreJunk" ), 
+	FString CmdLineParams = FString::Printf( TEXT( "%s %s %s %s %s%s -IgnoreJunk" ), 
 		*ModuleArg, 
-		BuildPlatformName, BuildConfigurationName, 
+		BuildTargetName, BuildPlatformName, BuildConfigurationName,
 		*ExtraArg, *InAdditionalCmdLineArgs );
 
 	const bool bInvocationSuccessful = InvokeUnrealBuildToolForCompile(CmdLineParams, Ar);
