Fix issue where hot reloading classes results in a fatal issue.

diff --git a/Engine/Source/Developer/HotReload/Private/HotReloadClassReinstancer.cpp b/Engine/Source/Developer/HotReload/Private/HotReloadClassReinstancer.cpp
--- a/Engine/Source/Developer/HotReload/Private/HotReloadClassReinstancer.cpp
+++ b/Engine/Source/Developer/HotReload/Private/HotReloadClassReinstancer.cpp
@@ -142,6 +142,11 @@ void FHotReloadClassReinstancer::SerializeCDOProperties(UObject* InObject, FHotR
 
 			return *this;
 		}
+		virtual FArchive& operator<<(FObjectPtr& InObj) override
+		{
+			// Invoke the method above
+			return FArchiveUObject::SerializeObjectPtr(*this, InObj);
+		}
 		/** Serializes an FName as its index and number */
 		virtual FArchive& operator<<(FName& InName) override
 		{
@@ -382,6 +387,11 @@ void FHotReloadClassReinstancer::UpdateDefaultProperties()
 			}
 			return *this;
 		}
+		virtual FArchive& operator<<(FObjectPtr& InObj) override
+		{
+			// Invoke the method above
+			return FArchiveUObject::SerializeObjectPtr(*this, InObj);
+		}
 		virtual FArchive& operator<<(FName& InName) override
 		{
 			FArchive& Ar = *this;
