Removing UE4 from identifiers

diff --git a/Engine/Source/Developer/HotReload/Private/HotReload.cpp b/Engine/Source/Developer/HotReload/Private/HotReload.cpp
--- a/Engine/Source/Developer/HotReload/Private/HotReload.cpp
+++ b/Engine/Source/Developer/HotReload/Private/HotReload.cpp
@@ -388,7 +388,7 @@ namespace HotReloadDefs
 
 IMPLEMENT_MODULE(FHotReloadModule, HotReload);
 
-namespace UE4HotReload_Private
+namespace UEHotReload_Private
 {
 	/**
 	 * Gets editor runs directory.
@@ -530,7 +530,7 @@ namespace UE4HotReload_Private
 
 void FHotReloadModule::StartupModule()
 {
-	UE4HotReload_Private::CreateFileThatIndicatesEditorRunIfNeeded();
+	UEHotReload_Private::CreateFileThatIndicatesEditorRunIfNeeded();
 
 	bIsHotReloadingFromEditor = false;
 
@@ -559,7 +559,7 @@ void FHotReloadModule::ShutdownModule()
 	FTicker::GetCoreTicker().RemoveTicker(TickerDelegateHandle);
 	ShutdownHotReloadWatcher();
 
-	UE4HotReload_Private::DeleteFileThatIndicatesEditorRunIfNeeded();
+	UEHotReload_Private::DeleteFileThatIndicatesEditorRunIfNeeded();
 }
 
 bool FHotReloadModule::Exec( UWorld* Inworld, const TCHAR* Cmd, FOutputDevice& Ar )
@@ -746,11 +746,11 @@ ECompilationResult::Type FHotReloadModule::DoHotReloadFromEditor(EHotReloadFlags
 {
 	// Get all game modules we want to compile
 	const FModuleManager& ModuleManager = FModuleManager::Get();
-	TArray<FString> GameModuleNames = UE4HotReload_Private::GetGameModuleNames(ModuleManager);
+	TArray<FString> GameModuleNames = UEHotReload_Private::GetGameModuleNames(ModuleManager);
 
 	ECompilationResult::Type Result = ECompilationResult::Unsupported;
 
-	UE4HotReload_Private::FPackagesAndDependentNames PackagesAndDependentNames = UE4HotReload_Private::SplitByPackagesAndDependentNames(GameModuleNames);
+	UEHotReload_Private::FPackagesAndDependentNames PackagesAndDependentNames = UEHotReload_Private::SplitByPackagesAndDependentNames(GameModuleNames);
 
 	// Analytics
 	double Duration = 0.0;
@@ -870,8 +870,8 @@ ECompilationResult::Type FHotReloadModule::RebindPackages(const TArray<UPackage*
 
 	// Get game packages
 	const FModuleManager& ModuleManager = FModuleManager::Get();
-	TArray<FString> GameModuleNames = UE4HotReload_Private::GetGameModuleNames(ModuleManager);
-	UE4HotReload_Private::FPackagesAndDependentNames PackagesAndDependentNames = UE4HotReload_Private::SplitByPackagesAndDependentNames(GameModuleNames);
+	TArray<FString> GameModuleNames = UEHotReload_Private::GetGameModuleNames(ModuleManager);
+	UEHotReload_Private::FPackagesAndDependentNames PackagesAndDependentNames = UEHotReload_Private::SplitByPackagesAndDependentNames(GameModuleNames);
 
 	// Get a set of source packages combined with game packages
 	TSet<UPackage*> PackagesIncludingGame(InPackages);
@@ -1073,7 +1073,7 @@ void FHotReloadModule::OnHotReloadBinariesChanged(const TArray<FFileChangeData>&
 	}
 
 	const FModuleManager& ModuleManager = FModuleManager::Get();
-	TMap<FName, FString> GameModuleFilenames = UE4HotReload_Private::GetGameModuleFilenames(ModuleManager);
+	TMap<FName, FString> GameModuleFilenames = UEHotReload_Private::GetGameModuleFilenames(ModuleManager);
 
 	if (GameModuleFilenames.Num() == 0)
 	{
@@ -1130,7 +1130,7 @@ void FHotReloadModule::OnHotReloadBinariesChanged(const TArray<FFileChangeData>&
 
 void FHotReloadModule::StripModuleSuffixFromFilename(FString& InOutModuleFilename, const FString& ModuleName)
 {
-	// First hyphen is where the UE4Edtior prefix ends
+	// First hyphen is where the UEEdtior prefix ends
 	int32 FirstHyphenIndex = INDEX_NONE;
 	if (InOutModuleFilename.FindChar('-', FirstHyphenIndex))
 	{
@@ -1289,7 +1289,7 @@ void FHotReloadModule::DoHotReloadFromIDE(const TMap<FName, FString>& NewModules
 
 	double Duration = 0.0;
 
-	TArray<FString> GameModuleNames = UE4HotReload_Private::GetGameModuleNames(ModuleManager);
+	TArray<FString> GameModuleNames = UEHotReload_Private::GetGameModuleNames(ModuleManager);
 	if (GameModuleNames.Num() > 0)
 	{
 		FScopedDurationTimer Timer(Duration);
@@ -1319,7 +1319,7 @@ void FHotReloadModule::DoHotReloadFromIDE(const TMap<FName, FString>& NewModules
 		}
 
 		SlowTask.EnterProgressFrame(10);
-		UE4HotReload_Private::FPackagesAndDependentNames PackagesAndDependentNames = UE4HotReload_Private::SplitByPackagesAndDependentNames(GameModuleNames);
+		UEHotReload_Private::FPackagesAndDependentNames PackagesAndDependentNames = UEHotReload_Private::SplitByPackagesAndDependentNames(GameModuleNames);
 		SlowTask.EnterProgressFrame(80);
 
 		NumPackagesToRebind = PackagesAndDependentNames.Packages.Num();
