Compile fix when WITH_HOT_RELOAD is 0

#rnx

diff --git a/Engine/Source/Developer/HotReload/Private/HotReload.cpp b/Engine/Source/Developer/HotReload/Private/HotReload.cpp
--- a/Engine/Source/Developer/HotReload/Private/HotReload.cpp
+++ b/Engine/Source/Developer/HotReload/Private/HotReload.cpp
@@ -85,6 +85,7 @@ namespace EThreeStateBool
 	}
 };
 
+#if WITH_HOT_RELOAD
 class FScopedHotReload
 {
 public:
@@ -108,6 +109,7 @@ public:
 private:
 	TUniquePtr<FReload>& UniquePtr;
 };
+#endif // WITH_HOT_RELOAD
 
 /**
  * Module for HotReload support
