Fixed HotReload to not invoke GC while the GIsInitialLoad flag is forced to true.

#rnx

diff --git a/Engine/Source/Developer/HotReload/Private/HotReload.cpp b/Engine/Source/Developer/HotReload/Private/HotReload.cpp
--- a/Engine/Source/Developer/HotReload/Private/HotReload.cpp
+++ b/Engine/Source/Developer/HotReload/Private/HotReload.cpp
@@ -783,9 +783,12 @@ ECompilationResult::Type FHotReloadModule::DoHotReloadInternal(const TMap<FName,
 	// Rebind the hot reload DLL 
 	FScopedHotReload Guard(Reload, Packages);
 	Reload->SetSendReloadCompleteNotification(false);
-	TGuardValue<bool> GuardIsInitialLoad(GIsInitialLoad, true);
 
-	CollectGarbage(GARBAGE_COLLECTION_KEEPFLAGS); // we create a new CDO in the transient package...this needs to go away before we try again.
+	// we create a new CDO in the transient package...this needs to go away before we try again.
+	CollectGarbage(GARBAGE_COLLECTION_KEEPFLAGS); 
+
+	// Pretend we are loading.  This must happen after GC
+	TGuardValue<bool> GuardIsInitialLoad(GIsInitialLoad, true);
 
 	// Load the new modules up
 	bool bReloadSucceeded = false;
diff --git a/Engine/Source/Editor/UnrealEd/Private/Kismet2/ReloadUtilities.cpp b/Engine/Source/Editor/UnrealEd/Private/Kismet2/ReloadUtilities.cpp
--- a/Engine/Source/Editor/UnrealEd/Private/Kismet2/ReloadUtilities.cpp
+++ b/Engine/Source/Editor/UnrealEd/Private/Kismet2/ReloadUtilities.cpp
@@ -1060,6 +1060,8 @@ void FReload::Finalize(bool bRunGC)
 	// Force GC to collect reinstanced objects
 	if (bRunGC)
 	{
+		// Make sure the GIsInitialLoad flag is false.  Otherwise GC does nothing
+		TGuardValue<bool> GuardIsInitialLoad(GIsInitialLoad, false);
 		CollectGarbage(GARBAGE_COLLECTION_KEEPFLAGS, true);
 	}
 }
