Added tagging for unassigned allocations

diff --git a/Engine/Plugins/Developer/CLionSourceCodeAccess/Source/CLionSourceCodeAccess/Private/CLionSourceCodeAccessModule.cpp b/Engine/Plugins/Developer/CLionSourceCodeAccess/Source/CLionSourceCodeAccess/Private/CLionSourceCodeAccessModule.cpp
--- a/Engine/Plugins/Developer/CLionSourceCodeAccess/Source/CLionSourceCodeAccess/Private/CLionSourceCodeAccessModule.cpp
+++ b/Engine/Plugins/Developer/CLionSourceCodeAccess/Source/CLionSourceCodeAccess/Private/CLionSourceCodeAccessModule.cpp
@@ -3,9 +3,12 @@
 #include "CLionSourceCodeAccessModule.h"
 #include "Modules/ModuleManager.h"
 #include "Features/IModularFeatures.h"
+#include "HAL/LowLevelMemTracker.h"
 
 #define LOCTEXT_NAMESPACE "CLionSourceCodeAccessor"
 
+LLM_DEFINE_TAG(CLionSourceCodeAccess);
+
 IMPLEMENT_MODULE(FCLionSourceCodeAccessModule, CLionSourceCodeAccess);
 
 void FCLionSourceCodeAccessModule::ShutdownModule()
@@ -16,6 +19,8 @@ void FCLionSourceCodeAccessModule::ShutdownModule()
 
 void FCLionSourceCodeAccessModule::StartupModule()
 {
+	LLM_SCOPE_BYTAG(CLionSourceCodeAccess);
+
 	// Quick forced check of availability before anyone touches the module
 	CLionSourceCodeAccessor.RefreshAvailability();
 
diff --git a/Engine/Plugins/Developer/RiderSourceCodeAccess/Source/RiderSourceCodeAccess/Private/RiderSourceCodeAccessorModule.cpp b/Engine/Plugins/Developer/RiderSourceCodeAccess/Source/RiderSourceCodeAccess/Private/RiderSourceCodeAccessorModule.cpp
--- a/Engine/Plugins/Developer/RiderSourceCodeAccess/Source/RiderSourceCodeAccess/Private/RiderSourceCodeAccessorModule.cpp
+++ b/Engine/Plugins/Developer/RiderSourceCodeAccess/Source/RiderSourceCodeAccess/Private/RiderSourceCodeAccessorModule.cpp
@@ -7,13 +7,17 @@
 
 #include "Modules/ModuleManager.h"
 #include "Features/IModularFeatures.h"
+#include "HAL/LowLevelMemTracker.h"
 
 #define LOCTEXT_NAMESPACE "RiderSourceCodeAccessor"
 
+LLM_DEFINE_TAG(RiderSourceCodeAccess);
+
 IMPLEMENT_MODULE(FRiderSourceCodeAccessModule, RiderSourceCodeAccess);
 
 void FRiderSourceCodeAccessModule::StartupModule()
 {
+	LLM_SCOPE_BYTAG(RiderSourceCodeAccess);
 	TArray<FInstallInfo> InstallInfos = FRiderPathLocator::CollectAllPaths().Array();
 	InstallInfos.Sort();
 	GenerateUprojectAccessors(InstallInfos);
diff --git a/Engine/Plugins/Developer/VisualStudioCodeSourceCodeAccess/Source/VisualStudioCodeSourceCodeAccess/Private/VisualStudioCodeSourceCodeAccessModule.cpp b/Engine/Plugins/Developer/VisualStudioCodeSourceCodeAccess/Source/VisualStudioCodeSourceCodeAccess/Private/VisualStudioCodeSourceCodeAccessModule.cpp
--- a/Engine/Plugins/Developer/VisualStudioCodeSourceCodeAccess/Source/VisualStudioCodeSourceCodeAccess/Private/VisualStudioCodeSourceCodeAccessModule.cpp
+++ b/Engine/Plugins/Developer/VisualStudioCodeSourceCodeAccess/Source/VisualStudioCodeSourceCodeAccess/Private/VisualStudioCodeSourceCodeAccessModule.cpp
@@ -3,6 +3,9 @@
 #include "VisualStudioCodeSourceCodeAccessModule.h"
 #include "Features/IModularFeatures.h"
 #include "Modules/ModuleManager.h"
+#include "HAL/LowLevelMemTracker.h"
+
+LLM_DEFINE_TAG(VisualStudioCodeSourceCodeAccess);
 
 IMPLEMENT_MODULE( FVisualStudioCodeSourceCodeAccessModule, VisualStudioCodeSourceCodeAccess );
 
@@ -15,6 +18,8 @@ FVisualStudioCodeSourceCodeAccessModule::FVisualStudioCodeSourceCodeAccessModule
 
 void FVisualStudioCodeSourceCodeAccessModule::StartupModule()
 {
+	LLM_SCOPE_BYTAG(VisualStudioCodeSourceCodeAccess);
+
 	VisualStudioCodeSourceCodeAccessor->Startup();
 
 	// Bind our source control provider to the editor
diff --git a/Engine/Plugins/Developer/VisualStudioSourceCodeAccess/Source/VisualStudioSourceCodeAccess/Private/VisualStudioSourceCodeAccessModule.cpp b/Engine/Plugins/Developer/VisualStudioSourceCodeAccess/Source/VisualStudioSourceCodeAccess/Private/VisualStudioSourceCodeAccessModule.cpp
--- a/Engine/Plugins/Developer/VisualStudioSourceCodeAccess/Source/VisualStudioSourceCodeAccess/Private/VisualStudioSourceCodeAccessModule.cpp
+++ b/Engine/Plugins/Developer/VisualStudioSourceCodeAccess/Source/VisualStudioSourceCodeAccess/Private/VisualStudioSourceCodeAccessModule.cpp
@@ -3,11 +3,14 @@
 #include "VisualStudioSourceCodeAccessModule.h"
 #include "Features/IModularFeatures.h"
 #include "Modules/ModuleManager.h"
+#include "HAL/LowLevelMemTracker.h"
 
 IMPLEMENT_MODULE( FVisualStudioSourceCodeAccessModule, VisualStudioSourceCodeAccess );
 
 #define LOCTEXT_NAMESPACE "VisualStudioSourceCodeAccessor"
 
+LLM_DEFINE_TAG(VisualStudioSourceCodeAccess);
+
 FVisualStudioSourceCodeAccessModule::FVisualStudioSourceCodeAccessModule()
 	: VisualStudioSourceCodeAccessor(MakeShareable(new FVisualStudioSourceCodeAccessor()))
 {
@@ -15,6 +18,8 @@ FVisualStudioSourceCodeAccessModule::FVisualStudioSourceCodeAccessModule()
 
 void FVisualStudioSourceCodeAccessModule::StartupModule()
 {
+	LLM_SCOPE_BYTAG(VisualStudioSourceCodeAccess);
+
 	VisualStudioSourceCodeAccessor->Startup();
 
 	// Add all the explicit version wrappers. If one of these is selected, UBT will generate project files in the appropriate format. Editor behavior is still to detect which version to use
diff --git a/Engine/Source/Developer/HotReload/Private/HotReload.cpp b/Engine/Source/Developer/HotReload/Private/HotReload.cpp
--- a/Engine/Source/Developer/HotReload/Private/HotReload.cpp
+++ b/Engine/Source/Developer/HotReload/Private/HotReload.cpp
@@ -33,6 +33,7 @@
 #include "ProfilingDebugging/ScopedTimers.h"
 #include "Interfaces/IPluginManager.h"
 #include "DesktopPlatformModule.h"
+#include "HAL/LowLevelMemTracker.h"
 #if WITH_ENGINE
 #include "Engine/Engine.h"
 #include "EngineAnalytics.h"
@@ -57,6 +58,8 @@ DEFINE_LOG_CATEGORY(LogHotReload);
 
 #define LOCTEXT_NAMESPACE "HotReload"
 
+LLM_DEFINE_TAG(HotReload);
+
 namespace EThreeStateBool
 {
 	enum Type
@@ -380,13 +383,13 @@ IMPLEMENT_MODULE(FHotReloadModule, HotReload);
 
 namespace HotReloadDefs
 {
-	const static FString CompilationInfoConfigSection("ModuleFileTracking");
+	static const TCHAR* CompilationInfoConfigSection = TEXT("ModuleFileTracking");
 
 	// These strings should match the values of the enum EModuleCompileMethod in ModuleManager.h
 	// and should be handled in ReadModuleCompilationInfoFromConfig() & WriteModuleCompilationInfoToConfig() below
-	const static FString CompileMethodRuntime("Runtime");
-	const static FString CompileMethodExternal("External");
-	const static FString CompileMethodUnknown("Unknown");
+	static const TCHAR* CompileMethodRuntime = TEXT("Runtime");
+	static const TCHAR* CompileMethodExternal = TEXT("External");
+	static const TCHAR* CompileMethodUnknown = TEXT("Unknown");
 
 	// Add one minute epsilon to timestamp comparision
 	const static FTimespan TimeStampEpsilon(0, 1, 0);
@@ -534,6 +537,8 @@ namespace UEHotReload_Private
 
 void FHotReloadModule::StartupModule()
 {
+	LLM_SCOPE_BYTAG(HotReload);
+
 	UEHotReload_Private::CreateFileThatIndicatesEditorRunIfNeeded();
 
 	bIsHotReloadingFromEditor = false;
@@ -605,6 +610,8 @@ void FHotReloadModule::SaveConfig()
 
 FString FHotReloadModule::GetModuleCompileMethod(FName InModuleName)
 {
+	LLM_SCOPE_BYTAG(HotReload);
+
 	if (!ModuleCompileData.Contains(InModuleName))
 	{
 		UpdateModuleCompileData(InModuleName);
@@ -1736,7 +1743,7 @@ void FHotReloadModule::UpdateModuleCompileData(FName ModuleName)
 void FHotReloadModule::ReadModuleCompilationInfoFromConfig(FName ModuleName, FModuleCompilationData& CompileData)
 {
 	FString DateTimeString;
-	if (GConfig->GetString(*HotReloadDefs::CompilationInfoConfigSection, *FString::Printf(TEXT("%s.TimeStamp"), *ModuleName.ToString()), DateTimeString, GEditorPerProjectIni))
+	if (GConfig->GetString(HotReloadDefs::CompilationInfoConfigSection, *FString::Printf(TEXT("%s.TimeStamp"), *ModuleName.ToString()), DateTimeString, GEditorPerProjectIni))
 	{
 		FDateTime TimeStamp;
 		if (!DateTimeString.IsEmpty() && FDateTime::Parse(DateTimeString, TimeStamp))
@@ -1745,13 +1752,13 @@ void FHotReloadModule::ReadModuleCompilationInfoFromConfig(FName ModuleName, FMo
 			CompileData.FileTimeStamp = TimeStamp;
 
 			FString CompileMethodString;
-			if (GConfig->GetString(*HotReloadDefs::CompilationInfoConfigSection, *FString::Printf(TEXT("%s.LastCompileMethod"), *ModuleName.ToString()), CompileMethodString, GEditorPerProjectIni))
+			if (GConfig->GetString(HotReloadDefs::CompilationInfoConfigSection, *FString::Printf(TEXT("%s.LastCompileMethod"), *ModuleName.ToString()), CompileMethodString, GEditorPerProjectIni))
 			{
-				if (CompileMethodString.Equals(HotReloadDefs::CompileMethodRuntime, ESearchCase::IgnoreCase))
+				if (FCString::Stricmp(*CompileMethodString, HotReloadDefs::CompileMethodRuntime) == 0)
 				{
 					CompileData.CompileMethod = EModuleCompileMethod::Runtime;
 				}
-				else if (CompileMethodString.Equals(HotReloadDefs::CompileMethodExternal, ESearchCase::IgnoreCase))
+				else if (FCString::Stricmp(*CompileMethodString, HotReloadDefs::CompileMethodExternal) == 0)
 				{
 					CompileData.CompileMethod = EModuleCompileMethod::External;
 				}
@@ -1768,9 +1775,9 @@ void FHotReloadModule::WriteModuleCompilationInfoToConfig(FName ModuleName, cons
 		DateTimeString = CompileData.FileTimeStamp.ToString();
 	}
 
-	GConfig->SetString(*HotReloadDefs::CompilationInfoConfigSection, *FString::Printf(TEXT("%s.TimeStamp"), *ModuleName.ToString()), *DateTimeString, GEditorPerProjectIni);
+	GConfig->SetString(HotReloadDefs::CompilationInfoConfigSection, *FString::Printf(TEXT("%s.TimeStamp"), *ModuleName.ToString()), *DateTimeString, GEditorPerProjectIni);
 
-	FString CompileMethodString = HotReloadDefs::CompileMethodUnknown;
+	const TCHAR* CompileMethodString = HotReloadDefs::CompileMethodUnknown;
 	if (CompileData.CompileMethod == EModuleCompileMethod::Runtime)
 	{
 		CompileMethodString = HotReloadDefs::CompileMethodRuntime;
@@ -1780,7 +1787,7 @@ void FHotReloadModule::WriteModuleCompilationInfoToConfig(FName ModuleName, cons
 		CompileMethodString = HotReloadDefs::CompileMethodExternal;
 	}
 
-	GConfig->SetString(*HotReloadDefs::CompilationInfoConfigSection, *FString::Printf(TEXT("%s.LastCompileMethod"), *ModuleName.ToString()), *CompileMethodString, GEditorPerProjectIni);
+	GConfig->SetString(HotReloadDefs::CompilationInfoConfigSection, *FString::Printf(TEXT("%s.LastCompileMethod"), *ModuleName.ToString()), CompileMethodString, GEditorPerProjectIni);
 }
 
 bool FHotReloadModule::GetModuleFileTimeStamp(FName ModuleName, FDateTime& OutFileTimeStamp) const
diff --git a/Engine/Source/Developer/SourceCodeAccess/Private/SourceCodeAccessModule.cpp b/Engine/Source/Developer/SourceCodeAccess/Private/SourceCodeAccessModule.cpp
--- a/Engine/Source/Developer/SourceCodeAccess/Private/SourceCodeAccessModule.cpp
+++ b/Engine/Source/Developer/SourceCodeAccess/Private/SourceCodeAccessModule.cpp
@@ -6,12 +6,15 @@
 #include "Features/IModularFeatures.h"
 #include "ISettingsModule.h"
 #include "SourceCodeAccessSettings.h"
+#include "HAL/LowLevelMemTracker.h"
 
 
 IMPLEMENT_MODULE( FSourceCodeAccessModule, SourceCodeAccess );
 
 #define LOCTEXT_NAMESPACE "SourceCodeAccessModule"
 
+LLM_DEFINE_TAG(SourceCodeAccess);
+
 static FName SourceCodeAccessorFeatureName(TEXT("SourceCodeAccessor"));
 
 
@@ -22,6 +25,8 @@ FSourceCodeAccessModule::FSourceCodeAccessModule()
 
 void FSourceCodeAccessModule::StartupModule()
 {
+	LLM_SCOPE_BYTAG(SourceCodeAccess);
+
 	GetMutableDefault<USourceCodeAccessSettings>()->LoadConfig();
 
 	// Register to check for source control features
diff --git a/Engine/Source/Developer/Windows/LiveCoding/Private/LiveCodingModule.cpp b/Engine/Source/Developer/Windows/LiveCoding/Private/LiveCodingModule.cpp
--- a/Engine/Source/Developer/Windows/LiveCoding/Private/LiveCodingModule.cpp
+++ b/Engine/Source/Developer/Windows/LiveCoding/Private/LiveCodingModule.cpp
@@ -20,6 +20,7 @@
 #include "Windows/WindowsHWrapper.h"
 #include "Algo/Sort.h"
 #include "Algo/BinarySearch.h"
+#include "HAL/LowLevelMemTracker.h"
 #if WITH_EDITOR
 	#include "Editor.h"
 	#include "Kismet2/ReloadUtilities.h"
@@ -45,6 +46,8 @@ IMPLEMENT_MODULE(FLiveCodingModule, LiveCoding)
 
 #define LOCTEXT_NAMESPACE "LiveCodingModule"
 
+LLM_DEFINE_TAG(LiveCoding);
+
 bool GIsCompileActive = false;
 bool GTriggerReload = false;
 bool GHasLoadedPatch = false;
@@ -179,6 +182,8 @@ FLiveCodingModule::~FLiveCodingModule()
 
 void FLiveCodingModule::StartupModule()
 {
+	LLM_SCOPE_BYTAG(LiveCoding);
+
 	// Register with NT to get dll nitrifications
 	FNtDllFunction RegisterFunc("LdrRegisterDllNotification");
 	RegisterFunc(0, OnDllNotification, this, &CallbackCookie);
