Exec commands changes:
- Port some development only exec commands to Exec_Dev.

diff --git a/Engine/Plugins/Online/OnlineSubsystemSteam/Source/Private/OnlineEncryptedAppTicketInterfaceSteam.cpp b/Engine/Plugins/Online/OnlineSubsystemSteam/Source/Private/OnlineEncryptedAppTicketInterfaceSteam.cpp
--- a/Engine/Plugins/Online/OnlineSubsystemSteam/Source/Private/OnlineEncryptedAppTicketInterfaceSteam.cpp
+++ b/Engine/Plugins/Online/OnlineSubsystemSteam/Source/Private/OnlineEncryptedAppTicketInterfaceSteam.cpp
@@ -231,7 +231,7 @@ bool FOnlineEncryptedAppTicketSteam::GetEncryptedAppTicket(TArray<uint8>& OutEnc
 	return false;
 }
 
-bool FOnlineEncryptedAppTicketSteam::Exec(UWorld* InWorld, const TCHAR* Cmd, FOutputDevice& Ar)
+bool FOnlineEncryptedAppTicketSteam::Exec_Dev(UWorld* InWorld, const TCHAR* Cmd, FOutputDevice& Ar)
 {
 #if !UE_BUILD_SHIPPING
 	if (FParse::Command(&Cmd, TEXT("RequestEncryptedAppTicket")))
diff --git a/Engine/Plugins/Online/OnlineSubsystemSteam/Source/Private/SocketSubsystemSteam.cpp b/Engine/Plugins/Online/OnlineSubsystemSteam/Source/Private/SocketSubsystemSteam.cpp
--- a/Engine/Plugins/Online/OnlineSubsystemSteam/Source/Private/SocketSubsystemSteam.cpp
+++ b/Engine/Plugins/Online/OnlineSubsystemSteam/Source/Private/SocketSubsystemSteam.cpp
@@ -733,7 +733,7 @@ bool FSocketSubsystemSteam::Tick(float DeltaTime)
 	return true;
 }
 
-bool FSocketSubsystemSteam::Exec(UWorld* InWorld, const TCHAR* Cmd, FOutputDevice& Ar)
+bool FSocketSubsystemSteam::Exec_Dev(UWorld* InWorld, const TCHAR* Cmd, FOutputDevice& Ar)
 {
 #if !UE_BUILD_SHIPPING
 	if (FParse::Command(&Cmd, TEXT("dumpsteamsessions")))
diff --git a/Engine/Plugins/Online/OnlineSubsystemSteam/Source/Private/SocketSubsystemSteam.h b/Engine/Plugins/Online/OnlineSubsystemSteam/Source/Private/SocketSubsystemSteam.h
--- a/Engine/Plugins/Online/OnlineSubsystemSteam/Source/Private/SocketSubsystemSteam.h
+++ b/Engine/Plugins/Online/OnlineSubsystemSteam/Source/Private/SocketSubsystemSteam.h
@@ -397,8 +397,9 @@ public:
 	 */
 	virtual bool IsSocketWaitSupported() const override { return false; }
 
+protected:
 	// FSelfRegisteringExec
-	virtual bool Exec(UWorld* InWorld, const TCHAR* Cmd, FOutputDevice& Ar) override;
+	virtual bool Exec_Dev(UWorld* InWorld, const TCHAR* Cmd, FOutputDevice& Ar) override;
 };
 
 /**
diff --git a/Engine/Plugins/Online/OnlineSubsystemSteam/Source/Public/OnlineEncryptedAppTicketInterfaceSteam.h b/Engine/Plugins/Online/OnlineSubsystemSteam/Source/Public/OnlineEncryptedAppTicketInterfaceSteam.h
--- a/Engine/Plugins/Online/OnlineSubsystemSteam/Source/Public/OnlineEncryptedAppTicketInterfaceSteam.h
+++ b/Engine/Plugins/Online/OnlineSubsystemSteam/Source/Public/OnlineEncryptedAppTicketInterfaceSteam.h
@@ -66,12 +66,13 @@ public:
 	 */
 	bool GetEncryptedAppTicket(TArray<uint8>& OutEncryptedData);
 
-	// FSelfRegisteringExec
-	virtual bool Exec(UWorld* InWorld, const TCHAR* Cmd, FOutputDevice& Ar) override;
-
 	/* Attach to this delegate to get notified about the encrypted application ticket results. */
 	FOnEncryptedAppTicketResponse OnEncryptedAppTicketResultDelegate;
 
+protected:
+	// FSelfRegisteringExec
+	virtual bool Exec_Dev(UWorld* InWorld, const TCHAR* Cmd, FOutputDevice& Ar) override;
+
 private:
 
 	/** Reference to the main Steam subsystem */
diff --git a/Engine/Plugins/Online/WebAuth/Source/Private/WebAuthModule.cpp b/Engine/Plugins/Online/WebAuth/Source/Private/WebAuthModule.cpp
--- a/Engine/Plugins/Online/WebAuth/Source/Private/WebAuthModule.cpp
+++ b/Engine/Plugins/Online/WebAuth/Source/Private/WebAuthModule.cpp
@@ -75,7 +75,7 @@ bool FWebAuthModule::HandleWebAuthCommand(const TCHAR* Cmd, FOutputDevice& Ar)
 	return true;
 }
 
-bool FWebAuthModule::Exec(UWorld* InWorld, const TCHAR* Cmd, FOutputDevice& Ar)
+bool FWebAuthModule::Exec_Dev(UWorld* InWorld, const TCHAR* Cmd, FOutputDevice& Ar)
 {
 	// Ignore any execs that don't start with WebAuth
 	if (FParse::Command(&Cmd, TEXT("WebAuth")))
diff --git a/Engine/Plugins/Online/WebAuth/Source/Public/WebAuthModule.h b/Engine/Plugins/Online/WebAuth/Source/Public/WebAuthModule.h
--- a/Engine/Plugins/Online/WebAuth/Source/Public/WebAuthModule.h
+++ b/Engine/Plugins/Online/WebAuth/Source/Public/WebAuthModule.h
@@ -23,19 +23,6 @@ class FWebAuthModule :
 
 public:
 
-	// FSelfRegisteringExec
-
-	/**
-	 * Handle exec commands starting with "WebAuth"
-	 *
-	 * @param InWorld	the world context
-	 * @param Cmd		the exec command being executed
-	 * @param Ar		the archive to log results to
-	 *
-	 * @return true if the handler consumed the input, false to continue searching handlers
-	 */
-	virtual bool Exec(UWorld* InWorld, const TCHAR* Cmd, FOutputDevice& Ar) override;
-
 	/**
 	 * Exec command handlers
 	 */
@@ -73,6 +60,20 @@ public:
 		return WebAuth != nullptr;
 	}
 
+protected:
+	// FSelfRegisteringExec
+
+	/**
+	 * Handle exec commands starting with "WebAuth"
+	 *
+	 * @param InWorld	the world context
+	 * @param Cmd		the exec command being executed
+	 * @param Ar		the archive to log results to
+	 *
+	 * @return true if the handler consumed the input, false to continue searching handlers
+	 */
+	virtual bool Exec_Dev(UWorld* InWorld, const TCHAR* Cmd, FOutputDevice& Ar) override;
+
 private:
 	// IModuleInterface
 
diff --git a/Engine/Plugins/Runtime/GameplayAbilities/Source/GameplayAbilities/Private/GameplayAbilitiesExec.cpp b/Engine/Plugins/Runtime/GameplayAbilities/Source/GameplayAbilities/Private/GameplayAbilitiesExec.cpp
--- a/Engine/Plugins/Runtime/GameplayAbilities/Source/GameplayAbilities/Private/GameplayAbilitiesExec.cpp
+++ b/Engine/Plugins/Runtime/GameplayAbilities/Source/GameplayAbilities/Private/GameplayAbilitiesExec.cpp
@@ -11,14 +11,15 @@ struct FGameplayAbilitiesExec : public FSelfRegisteringExec
 	{
 	}
 
+protected:
 	// Begin FExec Interface
-	virtual bool Exec(UWorld* Inworld, const TCHAR* Cmd, FOutputDevice& Ar) override;
+	virtual bool Exec_Dev(UWorld* Inworld, const TCHAR* Cmd, FOutputDevice& Ar) override;
 	// End FExec Interface
 };
 
 FGameplayAbilitiesExec GameplayAbilitiesExecInstance;
 
-bool FGameplayAbilitiesExec::Exec(UWorld* Inworld, const TCHAR* Cmd, FOutputDevice& Ar)
+bool FGameplayAbilitiesExec::Exec_Dev(UWorld* Inworld, const TCHAR* Cmd, FOutputDevice& Ar)
 {
 	if (Inworld == NULL)
 	{
diff --git a/Engine/Plugins/Runtime/Steam/SteamSockets/Source/SteamSockets/Private/SteamSocketsSubsystem.cpp b/Engine/Plugins/Runtime/Steam/SteamSockets/Source/SteamSockets/Private/SteamSocketsSubsystem.cpp
--- a/Engine/Plugins/Runtime/Steam/SteamSockets/Source/SteamSockets/Private/SteamSocketsSubsystem.cpp
+++ b/Engine/Plugins/Runtime/Steam/SteamSockets/Source/SteamSockets/Private/SteamSocketsSubsystem.cpp
@@ -507,7 +507,7 @@ bool FSteamSocketsSubsystem::Tick(float DeltaTime)
 	return true;
 }
 
-bool FSteamSocketsSubsystem::Exec(class UWorld* InWorld, const TCHAR* Cmd, FOutputDevice& Ar)
+bool FSteamSocketsSubsystem::Exec_Dev(class UWorld* InWorld, const TCHAR* Cmd, FOutputDevice& Ar)
 {
 	bool bIsHandled = false;
 
diff --git a/Engine/Plugins/Runtime/Steam/SteamSockets/Source/SteamSockets/Public/SteamSocketsSubsystem.h b/Engine/Plugins/Runtime/Steam/SteamSockets/Source/SteamSockets/Public/SteamSocketsSubsystem.h
--- a/Engine/Plugins/Runtime/Steam/SteamSockets/Source/SteamSockets/Public/SteamSocketsSubsystem.h
+++ b/Engine/Plugins/Runtime/Steam/SteamSockets/Source/SteamSockets/Public/SteamSocketsSubsystem.h
@@ -72,10 +72,6 @@ public:
 	virtual bool Tick(float DeltaTime) override;
 	//~ End FTickerObject Interface
 
-	//~ Begin FSelfRegisteringExec Interface
-	virtual bool Exec(class UWorld* InWorld, const TCHAR* Cmd, FOutputDevice& Ar) override;
-	//~ End FSelfRegisteringExec Interface
-
 	/** Returns if the application is using the SteamSocket relays */
 	bool IsUsingRelayNetwork() const { return bUseRelays; }
 
@@ -164,6 +160,10 @@ PACKAGE_SCOPE:
 	bool bShouldTestPeek;
 
 protected:
+	//~ Begin FSelfRegisteringExec Interface
+	virtual bool Exec_Dev(class UWorld* InWorld, const TCHAR* Cmd, FOutputDevice& Ar) override;
+	//~ End FSelfRegisteringExec Interface
+
 	void CleanSocketInformation(bool bForceClean);
 	void DumpSocketInformationMap() const;
 
diff --git a/Engine/Source/Developer/AutomationController/Private/AutomationCommandline.cpp b/Engine/Source/Developer/AutomationController/Private/AutomationCommandline.cpp
--- a/Engine/Source/Developer/AutomationController/Private/AutomationCommandline.cpp
+++ b/Engine/Source/Developer/AutomationController/Private/AutomationCommandline.cpp
@@ -478,8 +478,9 @@ public:
 		return false;
 	}
 	
+protected:
 	/** Console commands, see embeded usage statement **/
-	virtual bool Exec(UWorld*, const TCHAR* Cmd, FOutputDevice& Ar) override
+	virtual bool Exec_Dev(UWorld*, const TCHAR* Cmd, FOutputDevice& Ar) override
 	{
 		bool bHandled = false;
 		// Track whether we have a flag we care about passing through.
diff --git a/Engine/Source/Developer/HotReload/Private/HotReload.cpp b/Engine/Source/Developer/HotReload/Private/HotReload.cpp
--- a/Engine/Source/Developer/HotReload/Private/HotReload.cpp
+++ b/Engine/Source/Developer/HotReload/Private/HotReload.cpp
@@ -135,9 +135,6 @@ public:
 	virtual void StartupModule() override;
 	virtual void ShutdownModule() override;
 
-	/** FSelfRegisteringExec implementation */
-	virtual bool Exec( UWorld* Inworld, const TCHAR* Cmd, FOutputDevice& Ar ) override;
-
 	/** IHotReloadInterface implementation */
 	virtual void SaveConfig() override;
 	virtual bool RecompileModule(const FName InModuleName, FOutputDevice &Ar, ERecompileModuleFlags Flags) override;
@@ -156,6 +153,10 @@ public:
 	virtual FString GetModuleCompileMethod(FName InModuleName) override;
 	virtual bool IsAnyGameModuleLoaded() override;
 
+protected:
+	/** FSelfRegisteringExec implementation */
+	virtual bool Exec_Dev( UWorld* Inworld, const TCHAR* Cmd, FOutputDevice& Ar ) override;
+
 private:
 	/**
 	 * Enumerates compilation methods for modules.
@@ -571,7 +572,7 @@ void FHotReloadModule::ShutdownModule()
 	UEHotReload_Private::DeleteFileThatIndicatesEditorRunIfNeeded();
 }
 
-bool FHotReloadModule::Exec( UWorld* Inworld, const TCHAR* Cmd, FOutputDevice& Ar )
+bool FHotReloadModule::Exec_Dev( UWorld* Inworld, const TCHAR* Cmd, FOutputDevice& Ar )
 {
 #if !UE_BUILD_SHIPPING
 	if ( FParse::Command( &Cmd, TEXT( "Module" ) ) )
diff --git a/Engine/Source/Runtime/AIModule/Classes/EnvironmentQuery/EnvQueryManager.h b/Engine/Source/Runtime/AIModule/Classes/EnvironmentQuery/EnvQueryManager.h
--- a/Engine/Source/Runtime/AIModule/Classes/EnvironmentQuery/EnvQueryManager.h
+++ b/Engine/Source/Runtime/AIModule/Classes/EnvironmentQuery/EnvQueryManager.h
@@ -282,14 +282,14 @@ class AIMODULE_API UEnvQueryManager : public UAISubsystem, public FSelfRegisteri
 
 	static void SetAllowTimeSlicing(bool bAllowTimeSlicing);
 
-	//~ Begin FExec Interface
-	virtual bool Exec(UWorld* Inworld, const TCHAR* Cmd, FOutputDevice& Ar) override;
-	//~ End FExec Interface
-
 	/** Configure config variables during runtime */
 	void Configure(const FEnvQueryManagerConfig& NewConfig);
 
 protected:
+	//~ Begin FExec Interface
+	virtual bool Exec_Dev(UWorld* Inworld,const TCHAR* Cmd,FOutputDevice& Ar) override;
+	//~ End FExec Interface
+
 	friend UEnvQueryInstanceBlueprintWrapper;
 	TSharedPtr<FEnvQueryInstance> FindQueryInstance(const int32 QueryID);
 
diff --git a/Engine/Source/Runtime/AIModule/Private/AISystemExec.cpp b/Engine/Source/Runtime/AIModule/Private/AISystemExec.cpp
--- a/Engine/Source/Runtime/AIModule/Private/AISystemExec.cpp
+++ b/Engine/Source/Runtime/AIModule/Private/AISystemExec.cpp
@@ -12,14 +12,15 @@ struct FAISystemExec : public FSelfRegisteringExec
 	{
 	}
 
+protected:
 	// Begin FExec Interface
-	virtual bool Exec(UWorld* Inworld, const TCHAR* Cmd, FOutputDevice& Ar) override;
+	virtual bool Exec_Dev(UWorld* Inworld, const TCHAR* Cmd, FOutputDevice& Ar) override;
 	// End FExec Interface
 };
 
 FAISystemExec AISystemExecInstance;
 
-bool FAISystemExec::Exec(UWorld* Inworld, const TCHAR* Cmd, FOutputDevice& Ar)
+bool FAISystemExec::Exec_Dev(UWorld* Inworld, const TCHAR* Cmd, FOutputDevice& Ar)
 {
 	if (Inworld == NULL)
 	{
diff --git a/Engine/Source/Runtime/AIModule/Private/EnvironmentQuery/EnvQueryManager.cpp b/Engine/Source/Runtime/AIModule/Private/EnvironmentQuery/EnvQueryManager.cpp
--- a/Engine/Source/Runtime/AIModule/Private/EnvironmentQuery/EnvQueryManager.cpp
+++ b/Engine/Source/Runtime/AIModule/Private/EnvironmentQuery/EnvQueryManager.cpp
@@ -1084,7 +1084,7 @@ void UEnvQueryManager::SetAllowTimeSlicing(bool bAllowTimeSlicing)
 #endif
 }
 
-bool UEnvQueryManager::Exec(UWorld* Inworld, const TCHAR* Cmd, FOutputDevice& Ar)
+bool UEnvQueryManager::Exec_Dev(UWorld* Inworld, const TCHAR* Cmd, FOutputDevice& Ar)
 {
 #if USE_EQS_DEBUGGER
 	if (FParse::Command(&Cmd, TEXT("DumpEnvQueryStats")))
diff --git a/Engine/Source/Runtime/Core/Private/HAL/IPlatformFileLogWrapper.cpp b/Engine/Source/Runtime/Core/Private/HAL/IPlatformFileLogWrapper.cpp
--- a/Engine/Source/Runtime/Core/Private/HAL/IPlatformFileLogWrapper.cpp
+++ b/Engine/Source/Runtime/Core/Private/HAL/IPlatformFileLogWrapper.cpp
@@ -20,7 +20,7 @@ public:
 	{}
 
 	/** Console commands **/
-	virtual bool Exec(UWorld* InWorld, const TCHAR* Cmd, FOutputDevice& Ar) override
+	virtual bool Exec_Dev(UWorld* InWorld, const TCHAR* Cmd, FOutputDevice& Ar) override
 	{
 		if (FParse::Command(&Cmd, TEXT("LogFileDump")))
 		{
diff --git a/Engine/Source/Runtime/Core/Private/Modules/ModuleManager.cpp b/Engine/Source/Runtime/Core/Private/Modules/ModuleManager.cpp
--- a/Engine/Source/Runtime/Core/Private/Modules/ModuleManager.cpp
+++ b/Engine/Source/Runtime/Core/Private/Modules/ModuleManager.cpp
@@ -882,7 +882,7 @@ IModuleInterface* FModuleManager::GetModule( const FName InModuleName )
 	return nullptr;
 }
 
-bool FModuleManager::Exec( UWorld* Inworld, const TCHAR* Cmd, FOutputDevice& Ar )
+bool FModuleManager::Exec_Dev( UWorld* Inworld, const TCHAR* Cmd, FOutputDevice& Ar )
 {
 #if !UE_BUILD_SHIPPING
 	if ( FParse::Command( &Cmd, TEXT( "Module" ) ) )
diff --git a/Engine/Source/Runtime/Core/Public/Modules/ModuleManager.h b/Engine/Source/Runtime/Core/Public/Modules/ModuleManager.h
--- a/Engine/Source/Runtime/Core/Public/Modules/ModuleManager.h
+++ b/Engine/Source/Runtime/Core/Public/Modules/ModuleManager.h
@@ -506,11 +506,9 @@ public:
 		return IsPackageLoaded;
 	}
 
-
+protected:
 	// FSelfRegisteringExec interface.
-
-	virtual bool Exec( UWorld* Inworld, const TCHAR* Cmd, FOutputDevice& Ar ) override;
-
+	virtual bool Exec_Dev( UWorld* Inworld, const TCHAR* Cmd, FOutputDevice& Ar ) override;
 
 private:
 	friend TOptional<FModuleManager>& UE::Core::Private::GetModuleManagerSingleton();
diff --git a/Engine/Source/Runtime/CoreUObject/Private/UObject/LinkerManager.cpp b/Engine/Source/Runtime/CoreUObject/Private/UObject/LinkerManager.cpp
--- a/Engine/Source/Runtime/CoreUObject/Private/UObject/LinkerManager.cpp
+++ b/Engine/Source/Runtime/CoreUObject/Private/UObject/LinkerManager.cpp
@@ -26,7 +26,7 @@ FLinkerManager::~FLinkerManager()
 {
 }
 
-bool FLinkerManager::Exec(class UWorld* InWorld, const TCHAR* Cmd, FOutputDevice& Ar)
+bool FLinkerManager::Exec_Dev(class UWorld* InWorld, const TCHAR* Cmd, FOutputDevice& Ar)
 {
 #if !UE_BUILD_SHIPPING
 	if (FParse::Command(&Cmd, TEXT("LinkerLoadList")))
diff --git a/Engine/Source/Runtime/CoreUObject/Private/UObject/LinkerManager.h b/Engine/Source/Runtime/CoreUObject/Private/UObject/LinkerManager.h
--- a/Engine/Source/Runtime/CoreUObject/Private/UObject/LinkerManager.h
+++ b/Engine/Source/Runtime/CoreUObject/Private/UObject/LinkerManager.h
@@ -107,9 +107,6 @@ public:
 	}
 #endif
 
-	// FSelfRegisteringExec interface
-	virtual bool Exec(class UWorld* InWorld, const TCHAR* Cmd, FOutputDevice& Ar) override;
-
 	/** Reset only the linker exports associated with this package. */
 	void ResetLinkerExports(UPackage* InPackage);
 
@@ -137,6 +134,10 @@ public:
 	/** Adds a linker to deferred cleanup list */
 	void RemoveLinker(FLinkerLoad* Linker);
 
+protected:
+	// FSelfRegisteringExec interface
+	virtual bool Exec_Dev(class UWorld* InWorld,const TCHAR* Cmd,FOutputDevice& Ar) override;
+
 private:
 
 	/** Map of packages to their open linkers **/
diff --git a/Engine/Source/Runtime/Engine/Classes/AI/Navigation/AvoidanceManager.h b/Engine/Source/Runtime/Engine/Classes/AI/Navigation/AvoidanceManager.h
--- a/Engine/Source/Runtime/Engine/Classes/AI/Navigation/AvoidanceManager.h
+++ b/Engine/Source/Runtime/Engine/Classes/AI/Navigation/AvoidanceManager.h
@@ -173,16 +173,15 @@ class ENGINE_API UAvoidanceManager : public UObject, public FSelfRegisteringExec
 	void HandleToggleAvoidance( const TCHAR* Cmd, FOutputDevice& Ar );
 #endif
 
-	//~ Begin FExec Interface
-	virtual bool Exec(UWorld* Inworld, const TCHAR* Cmd, FOutputDevice& Ar) override;
-	//~ End FExec Interface
-
 	void SetNavEdgeProvider(INavEdgeProviderInterface* InEdgeProvider);
 
 	/** Used to manually mark data associated with given ID as 'dead' (and reusable) */
 	void RemoveAvoidanceObject(const int32 AvoidanceUID);
 
 protected:
+	//~ Begin FExec Interface
+	virtual bool Exec_Dev(UWorld* Inworld,const TCHAR* Cmd,FOutputDevice& Ar) override;
+	//~ End FExec Interface
 
 	/** Handle for efficient management of RemoveOutdatedObjects timer */
 	FTimerHandle TimerHandle_RemoveOutdatedObjects;
diff --git a/Engine/Source/Runtime/Engine/Classes/Engine/NetDriver.h b/Engine/Source/Runtime/Engine/Classes/Engine/NetDriver.h
--- a/Engine/Source/Runtime/Engine/Classes/Engine/NetDriver.h
+++ b/Engine/Source/Runtime/Engine/Classes/Engine/NetDriver.h
@@ -1251,7 +1251,7 @@ public:
 	//~ End UObject Interface.
 
 	//~ Begin FExec Interface
-
+protected:
 	/**
 	 * Handle exec commands
 	 *
@@ -1261,11 +1261,12 @@ public:
 	 *
 	 * @return true if the handler consumed the input, false to continue searching handlers
 	 */
-	ENGINE_API virtual bool Exec(UWorld* InWorld, const TCHAR* Cmd, FOutputDevice& Ar=*GLog) override;
+	ENGINE_API virtual bool Exec_Dev(UWorld* InWorld, const TCHAR* Cmd, FOutputDevice& Ar=*GLog) override;
+//~ End FExec Interface.
 
-	ENGINE_API ENetMode	GetNetMode() const;
+public:
 
-	//~ End FExec Interface.
+	ENGINE_API ENetMode	GetNetMode() const;
 
 	/** 
 	 * Returns true if this net driver is valid for the current configuration.
diff --git a/Engine/Source/Runtime/Engine/Private/AI/Navigation/AvoidanceManager.cpp b/Engine/Source/Runtime/Engine/Private/AI/Navigation/AvoidanceManager.cpp
--- a/Engine/Source/Runtime/Engine/Private/AI/Navigation/AvoidanceManager.cpp
+++ b/Engine/Source/Runtime/Engine/Private/AI/Navigation/AvoidanceManager.cpp
@@ -645,7 +645,7 @@ void UAvoidanceManager::HandleToggleAvoidance( const TCHAR* Cmd, FOutputDevice&
 
 #endif
 
-bool UAvoidanceManager::Exec(UWorld* Inworld, const TCHAR* Cmd, FOutputDevice& Ar)
+bool UAvoidanceManager::Exec_Dev(UWorld* Inworld, const TCHAR* Cmd, FOutputDevice& Ar)
 {
 #if !(UE_BUILD_SHIPPING || UE_BUILD_TEST)
 	if (FParse::Command(&Cmd, TEXT("AvoidanceDisplayAll")))
diff --git a/Engine/Source/Runtime/Engine/Private/AudioDevice.cpp b/Engine/Source/Runtime/Engine/Private/AudioDevice.cpp
--- a/Engine/Source/Runtime/Engine/Private/AudioDevice.cpp
+++ b/Engine/Source/Runtime/Engine/Private/AudioDevice.cpp
@@ -2211,7 +2211,7 @@ void FAudioDevice::SetMixDebugState(EDebugState InDebugState)
 	DebugState = InDebugState;
 }
 
-bool FAudioDevice::Exec(UWorld* InWorld, const TCHAR* Cmd, FOutputDevice& Ar)
+bool FAudioDevice::Exec_Dev(UWorld* InWorld, const TCHAR* Cmd, FOutputDevice& Ar)
 {
 #if !UE_BUILD_SHIPPING
 	auto ParseAudioExecCmd = [](const TCHAR** InCmd, const TCHAR* InMatch)
diff --git a/Engine/Source/Runtime/Engine/Private/Collision/WorldCollision.cpp b/Engine/Source/Runtime/Engine/Private/Collision/WorldCollision.cpp
--- a/Engine/Source/Runtime/Engine/Private/Collision/WorldCollision.cpp
+++ b/Engine/Source/Runtime/Engine/Private/Collision/WorldCollision.cpp
@@ -516,9 +516,9 @@ bool UWorld::ComponentSweepMultiByChannel(TArray<struct FHitResult>& OutHits, cl
 
 static class FCollisionExec : private FSelfRegisteringExec
 {
-public:
+protected:
 	/** Console commands, see embeded usage statement **/
-	virtual bool Exec( UWorld* InWorld, const TCHAR* Cmd, FOutputDevice& Ar ) override
+	virtual bool Exec_Dev( UWorld* InWorld, const TCHAR* Cmd, FOutputDevice& Ar ) override
 	{
 #if ENABLE_COLLISION_ANALYZER
 		if (FParse::Command(&Cmd, TEXT("CANALYZER")))
diff --git a/Engine/Source/Runtime/Engine/Private/NetDriver.cpp b/Engine/Source/Runtime/Engine/Private/NetDriver.cpp
--- a/Engine/Source/Runtime/Engine/Private/NetDriver.cpp
+++ b/Engine/Source/Runtime/Engine/Private/NetDriver.cpp
@@ -3133,8 +3133,7 @@ void HandleNetFlushAllDormancy(UNetDriver* InNetDriver, UWorld* InWorld)
 }
 #endif
 
-
-bool UNetDriver::Exec( UWorld* InWorld, const TCHAR* Cmd, FOutputDevice& Ar )
+bool UNetDriver::Exec_Dev( UWorld* InWorld, const TCHAR* Cmd, FOutputDevice& Ar )
 {
 #if !UE_BUILD_SHIPPING
 	int32 PacketLossBurstMilliseconds = 0;
diff --git a/Engine/Source/Runtime/Engine/Private/VisualLogger/VisualLogger.cpp b/Engine/Source/Runtime/Engine/Private/VisualLogger/VisualLogger.cpp
--- a/Engine/Source/Runtime/Engine/Private/VisualLogger/VisualLogger.cpp
+++ b/Engine/Source/Runtime/Engine/Private/VisualLogger/VisualLogger.cpp
@@ -1132,9 +1132,9 @@ FCustomVersionRegistration GVisualLoggerVersion(EVisualLoggerVersion::GUID, EVis
 
 class FLogVisualizerExec : private FSelfRegisteringExec
 {
-public:
+protected:
 	/** Console commands, see embedded usage statement **/
-	virtual bool Exec(UWorld* InWorld, const TCHAR* Cmd, FOutputDevice& Ar) override
+	virtual bool Exec_Dev(UWorld* InWorld, const TCHAR* Cmd, FOutputDevice& Ar) override
 	{
 		if (FParse::Command(&Cmd, TEXT("VISLOG")))
 		{
diff --git a/Engine/Source/Runtime/Engine/Public/AudioDevice.h b/Engine/Source/Runtime/Engine/Public/AudioDevice.h
--- a/Engine/Source/Runtime/Engine/Public/AudioDevice.h
+++ b/Engine/Source/Runtime/Engine/Public/AudioDevice.h
@@ -414,12 +414,13 @@ DECLARE_MULTICAST_DELEGATE_OneParam(FOnAudioDevicePostRender, const FAudioDevice
 
 class ENGINE_API FAudioDevice : public FExec
 {
-public:
+protected:
 
 	//Begin FExec Interface
-	virtual bool Exec(UWorld* InWorld, const TCHAR* Cmd, FOutputDevice& Ar = *GLog) override;
+	virtual bool Exec_Dev(UWorld* InWorld, const TCHAR* Cmd, FOutputDevice& Ar = *GLog) override;
 	//End FExec Interface
 
+public:
 #if !UE_BUILD_SHIPPING
 	UAudioComponent* GetTestComponent(UWorld* InWorld);
 	void StopTestComponent();
diff --git a/Engine/Source/Runtime/NavigationSystem/Private/NavMesh/RecastNavMeshGenerator.cpp b/Engine/Source/Runtime/NavigationSystem/Private/NavMesh/RecastNavMeshGenerator.cpp
--- a/Engine/Source/Runtime/NavigationSystem/Private/NavMesh/RecastNavMeshGenerator.cpp
+++ b/Engine/Source/Runtime/NavigationSystem/Private/NavMesh/RecastNavMeshGenerator.cpp
@@ -7194,9 +7194,9 @@ void FRecastNavMeshGenerator::ExportNavigationData(const FString& FileName) cons
 
 static class FNavigationGeomExec : private FSelfRegisteringExec
 {
-public:
+protected:
 	/** Console commands, see embeded usage statement **/
-	virtual bool Exec( UWorld* InWorld, const TCHAR* Cmd, FOutputDevice& Ar ) override
+	virtual bool Exec_Dev( UWorld* InWorld, const TCHAR* Cmd, FOutputDevice& Ar ) override
 	{
 		bool bExported = false;
 #if ALLOW_DEBUG_FILES && ENABLE_VISUAL_LOG
diff --git a/Engine/Source/Runtime/PakFile/Private/IPlatformFilePak.cpp b/Engine/Source/Runtime/PakFile/Private/IPlatformFilePak.cpp
--- a/Engine/Source/Runtime/PakFile/Private/IPlatformFilePak.cpp
+++ b/Engine/Source/Runtime/PakFile/Private/IPlatformFilePak.cpp
@@ -7190,8 +7190,9 @@ public:
 		: PlatformFile(InPlatformFile)
 	{}
 
+protected:
 	/** Console commands **/
-	virtual bool Exec(UWorld* InWorld, const TCHAR* Cmd, FOutputDevice& Ar) override
+	virtual bool Exec_Dev(UWorld* InWorld, const TCHAR* Cmd, FOutputDevice& Ar) override
 	{
 		if (FParse::Command(&Cmd, TEXT("Mount")))
 		{
diff --git a/Engine/Source/Runtime/VulkanRHI/Private/VulkanSwapChain.cpp b/Engine/Source/Runtime/VulkanRHI/Private/VulkanSwapChain.cpp
--- a/Engine/Source/Runtime/VulkanRHI/Private/VulkanSwapChain.cpp
+++ b/Engine/Source/Runtime/VulkanRHI/Private/VulkanSwapChain.cpp
@@ -84,7 +84,7 @@ bool GSimulateSuboptimalSurfaceInNextTick = false;
 // A self registering exec helper to check for the VULKAN_* commands.
 class FVulkanCommandsHelper : public FSelfRegisteringExec
 {
-	virtual bool Exec(class UWorld* InWorld, const TCHAR* Cmd, FOutputDevice& Ar)
+	virtual bool Exec_Dev(class UWorld* InWorld, const TCHAR* Cmd, FOutputDevice& Ar)
 	{
 		if (FParse::Command(&Cmd, TEXT("VULKAN_SIMULATE_LOST_SURFACE")))
 		{
