Fix Hot Reload failing on AppleSilicon Macs as UBT compiles the changed module as x86_64.

diff --git a/Engine/Source/Developer/HotReload/Private/HotReload.cpp b/Engine/Source/Developer/HotReload/Private/HotReload.cpp
--- a/Engine/Source/Developer/HotReload/Private/HotReload.cpp
+++ b/Engine/Source/Developer/HotReload/Private/HotReload.cpp
@@ -1402,6 +1402,11 @@ FString FHotReloadModule::MakeUBTArgumentsForModuleCompiling()
 		FString FullProjectPath = FPaths::ConvertRelativePathToFull(FPaths::GetProjectFilePath());
 		AdditionalArguments += FString::Printf(TEXT("\"%s\" "), *FullProjectPath);
 	}
+	
+#if PLATFORM_MAC_ARM64
+	AdditionalArguments += "-architecture=arm64";
+#endif
+
 	return AdditionalArguments;
 }
 
