Disable HotReload module in unattended runs and commandlets.

- Live edits to code are not expected to happen for unattended runs or commandlets.
- The change primarily targets reducing cooker (+cook workers) memory usage.

diff --git a/Engine/Source/Developer/HotReload/Private/HotReload.cpp b/Engine/Source/Developer/HotReload/Private/HotReload.cpp
--- a/Engine/Source/Developer/HotReload/Private/HotReload.cpp
+++ b/Engine/Source/Developer/HotReload/Private/HotReload.cpp
@@ -536,10 +536,24 @@ namespace UEHotReload_Private
 	}
 }
 
+namespace HotReloadModule
+{
+	inline bool DisableHotReloadForThisRun()
+	{
+		// HotReload will be disabled when either running in -unattended mode or running a commandlet.
+		return FApp::IsUnattended() || IsRunningCommandlet();
+	}
+}
+
 void FHotReloadModule::StartupModule()
 {
 	LLM_SCOPE_BYTAG(HotReload);
 
+	if (HotReloadModule::DisableHotReloadForThisRun())
+	{
+		return;
+	}
+
 	UEHotReload_Private::CreateFileThatIndicatesEditorRunIfNeeded();
 
 	bIsHotReloadingFromEditor = false;
@@ -566,6 +580,11 @@ void FHotReloadModule::StartupModule()
 
 void FHotReloadModule::ShutdownModule()
 {
+	if (HotReloadModule::DisableHotReloadForThisRun())
+	{
+		return;
+	}
+
 	FTSTicker::GetCoreTicker().RemoveTicker(TickerDelegateHandle);
 	ShutdownHotReloadWatcher();
 
