Core: Cleanup of 5.0 deprecations in Core and CoreUObject.
A strict code removal operation.
Skip removal of enum values to minimize risk.
Skip removal of pure vitual FEvent::Create(bool bIsManualReset = false) that is still in use.
Skip removal of FTopLevelAssetPath(FName InPath) constructor that is still in use.
Skip removal of FEditorBulkData::SetVirtualizationOptOut (bool bOptOut) that is combined with a UE_ENABLE_VIRTUALIZATION_TOGGLE define.
Skip removal of GPackageFileUE4Version and GPackageFileLicenseeUE4Version in ObjectVersion.h that is under strict lockdown control.

#rnx
#tests BuildCookRun Lyra, Editor boot

diff --git a/Engine/Source/Developer/HotReload/Private/HotReload.cpp b/Engine/Source/Developer/HotReload/Private/HotReload.cpp
--- a/Engine/Source/Developer/HotReload/Private/HotReload.cpp
+++ b/Engine/Source/Developer/HotReload/Private/HotReload.cpp
@@ -140,14 +140,8 @@ public:
 	virtual bool RecompileModule(const FName InModuleName, FOutputDevice &Ar, ERecompileModuleFlags Flags) override;
 	virtual bool IsCurrentlyCompiling() const override { return ModuleCompileProcessHandle.IsValid(); }
 	virtual void RequestStopCompilation() override { bRequestCancelCompilation = true; }
-	PRAGMA_DISABLE_DEPRECATION_WARNINGS
-	virtual void AddHotReloadFunctionRemap(FNativeFuncPtr NewFunctionPointer, FNativeFuncPtr OldFunctionPointer) override;
-	PRAGMA_ENABLE_DEPRECATION_WARNINGS
 	virtual ECompilationResult::Type RebindPackages(const TArray<UPackage*>& Packages, EHotReloadFlags Flags, FOutputDevice &Ar) override;
 	virtual ECompilationResult::Type DoHotReloadFromEditor(EHotReloadFlags Flags) override;
-	PRAGMA_DISABLE_DEPRECATION_WARNINGS
-	virtual FHotReloadEvent& OnHotReload() override { return HotReloadEvent; }
-	PRAGMA_ENABLE_DEPRECATION_WARNINGS
 	virtual FModuleCompilerStartedEvent& OnModuleCompilerStarted() override { return ModuleCompilerStartedEvent; }
 	virtual FModuleCompilerFinishedEvent& OnModuleCompilerFinished() override { return ModuleCompilerFinishedEvent; }
 	virtual FString GetModuleCompileMethod(FName InModuleName) override;
@@ -334,9 +328,6 @@ private:
 	/** Modules that have been recently recompiled from the editor **/
 	TSet<FName> ModulesRecentlyCompiledInTheEditor;
 
-	/** Delegate broadcast when a module has been hot-reloaded */
-	FHotReloadEvent HotReloadEvent;
-
 	/** Array of modules that we're currently recompiling */
 	TArray< FName > ModulesBeingCompiled;
 
@@ -757,7 +748,6 @@ bool FHotReloadModule::RecompileModule(const FName InModuleName, FOutputDevice &
 
 	if (!!(Flags & ERecompileModuleFlags::ForceCodeProject))
 	{
-		HotReloadEvent.Broadcast( false );
 		FCoreUObjectDelegates::ReloadCompleteDelegate.Broadcast(EReloadCompleteReason::HotReloadManual);
 	}
 
@@ -767,12 +757,6 @@ bool FHotReloadModule::RecompileModule(const FName InModuleName, FOutputDevice &
 #endif // WITH_HOT_RELOAD
 }
 
-/** Adds and entry for the UFunction native pointer remap table */
-void FHotReloadModule::AddHotReloadFunctionRemap(FNativeFuncPtr NewFunctionPointer, FNativeFuncPtr OldFunctionPointer)
-{
-	ReloadNotifyFunctionRemap(NewFunctionPointer, OldFunctionPointer);
-}
-
 ECompilationResult::Type FHotReloadModule::DoHotReloadFromEditor(EHotReloadFlags Flags)
 {
 	// Get all game modules we want to compile
@@ -882,7 +866,6 @@ ECompilationResult::Type FHotReloadModule::DoHotReloadInternal(const TMap<FName,
 	}
 
 
-	HotReloadEvent.Broadcast( !bIsHotReloadingFromEditor);
 	FCoreUObjectDelegates::ReloadCompleteDelegate.Broadcast(bIsHotReloadingFromEditor ? EReloadCompleteReason::HotReloadManual : EReloadCompleteReason::HotReloadAutomatic);
 
 	HotReloadAr.Logf(ELogVerbosity::Display, TEXT("HotReload took %4.1fs."), FPlatformTime::Seconds() - HotReloadStartTime);
diff --git a/Engine/Source/Editor/UnrealEd/Private/EditorEngine.cpp b/Engine/Source/Editor/UnrealEd/Private/EditorEngine.cpp
--- a/Engine/Source/Editor/UnrealEd/Private/EditorEngine.cpp
+++ b/Engine/Source/Editor/UnrealEd/Private/EditorEngine.cpp
@@ -4514,9 +4514,6 @@ FSavePackageResultStruct UEditorEngine::Save(UPackage* InOuter, UObject* InAsset
 
 	SlowTask.EnterProgressFrame(70);
 
-	PRAGMA_DISABLE_DEPRECATION_WARNINGS;
-	UPackage::PreSavePackageEvent.Broadcast(InOuter);
-	PRAGMA_ENABLE_DEPRECATION_WARNINGS;
 	UPackage::PreSavePackageWithContextEvent.Broadcast(InOuter, FObjectPreSaveContext(ObjectSaveContext));
 	FSavePackageResultStruct Result = UPackage::Save(InOuter, Asset, Filename, SaveArgs);
 
diff --git a/Engine/Source/Runtime/Core/Private/GenericPlatform/GenericPlatformCrashContext.cpp b/Engine/Source/Runtime/Core/Private/GenericPlatform/GenericPlatformCrashContext.cpp
--- a/Engine/Source/Runtime/Core/Private/GenericPlatform/GenericPlatformCrashContext.cpp
+++ b/Engine/Source/Runtime/Core/Private/GenericPlatform/GenericPlatformCrashContext.cpp
@@ -1671,11 +1671,6 @@ void FGenericCrashContext::CaptureThreadPortableCallStack(const uint64 ThreadId,
 	SetPortableCallStack(StackTrace, StackTraceDepth);
 }
 
-void FGenericCrashContext::CapturePortableCallStack(int32 NumStackFramesToIgnore, void* Context)
-{
-	CapturePortableCallStack(nullptr, Context);
-}
-
 void FGenericCrashContext::SetPortableCallStack(const uint64* StackFrames, int32 NumStackFrames)
 {
 	GetPortableCallStack(StackFrames, NumStackFrames, CallStack);
diff --git a/Engine/Source/Runtime/Core/Private/Misc/FileHelper.cpp b/Engine/Source/Runtime/Core/Private/Misc/FileHelper.cpp
--- a/Engine/Source/Runtime/Core/Private/Misc/FileHelper.cpp
+++ b/Engine/Source/Runtime/Core/Private/Misc/FileHelper.cpp
@@ -802,39 +802,6 @@ void FFileHelper::GenerateDateTimeBasedBitmapFilename(const FString& Pattern, co
 	OutFilename = FString::Printf(TEXT("%s_%s.%s"), *Pattern, *FDateTime::Now().ToString(), *Extension);
 }
 
-/**
- * Saves a 24Bit BMP file to disk
- * 
- * @param Pattern filename with path, must not be 0, if with "bmp" extension (e.g. "out.bmp") the filename stays like this, if without (e.g. "out") automatic index numbers are addended (e.g. "out00002.bmp")
- * @param DataWidth - Width of the bitmap supplied in Data >0
- * @param DataHeight - Height of the bitmap supplied in Data >0
- * @param Data must not be 0
- * @param SubRectangle optional, specifies a sub-rectangle of the source image to save out. If NULL, the whole bitmap is saved
- * @param FileManager must not be 0
- * @param OutFilename optional, if specified filename will be output
- * @param ChannelMask optional, specifies a specific channel to write out (will be written out to all channels gray scale).
- *
- * @return true if success
- */
-PRAGMA_DISABLE_DEPRECATION_WARNINGS
-bool FFileHelper::CreateBitmap(const TCHAR* Pattern, int32 SourceWidth, int32 SourceHeight, const FColor* Data, FIntRect* SubRectangle, IFileManager* FileManager /*= &IFileManager::Get()*/, FString* OutFilename /*= NULL*/, bool bInWriteAlpha /*= false*/, EChannelMask ChannelMask /*= All */)
-{
-	EColorChannel ColorChannel = EColorChannel::All;
-	if (ChannelMask != EChannelMask::All)
-	{
-		switch (ChannelMask)
-		{
-		case EChannelMask::R:	ColorChannel = EColorChannel::R;	break;
-		case EChannelMask::G:	ColorChannel = EColorChannel::G;	break;
-		case EChannelMask::B:	ColorChannel = EColorChannel::B;	break;
-		case EChannelMask::A:	ColorChannel = EColorChannel::A;	break;
-		}
-	}
-
-	return CreateBitmap(Pattern, SourceWidth, SourceHeight, Data, SubRectangle, FileManager, OutFilename, bInWriteAlpha, ColorChannel);
-}
-PRAGMA_ENABLE_DEPRECATION_WARNINGS
-
 bool FFileHelper::CreateBitmap( const TCHAR* Pattern, int32 SourceWidth, int32 SourceHeight, const FColor* Data, FIntRect* SubRectangle, IFileManager* FileManager, FString* OutFilename, bool bInWriteAlpha, EColorChannel ColorChannel)
 {
 	FIntRect Src(0, 0, SourceWidth, SourceHeight);
diff --git a/Engine/Source/Runtime/Core/Public/GenericPlatform/GenericPlatformCrashContext.h b/Engine/Source/Runtime/Core/Public/GenericPlatform/GenericPlatformCrashContext.h
--- a/Engine/Source/Runtime/Core/Public/GenericPlatform/GenericPlatformCrashContext.h
+++ b/Engine/Source/Runtime/Core/Public/GenericPlatform/GenericPlatformCrashContext.h
@@ -716,9 +716,6 @@ public:
 	 */
 	CORE_API void CaptureThreadPortableCallStack(const uint64 ThreadId, void* Context);
 
-	UE_DEPRECATED(5.0, "")
-	CORE_API void CapturePortableCallStack(int32 NumStackFramesToIgnore, void* Context);
-	
 	/** Sets the portable callstack to a specified stack */
 	CORE_API virtual void SetPortableCallStack(const uint64* StackFrames, int32 NumStackFrames);
 
diff --git a/Engine/Source/Runtime/Core/Public/GenericPlatform/GenericPlatformProcess.h b/Engine/Source/Runtime/Core/Public/GenericPlatform/GenericPlatformProcess.h
--- a/Engine/Source/Runtime/Core/Public/GenericPlatform/GenericPlatformProcess.h
+++ b/Engine/Source/Runtime/Core/Public/GenericPlatform/GenericPlatformProcess.h
@@ -770,7 +770,8 @@ struct FGenericPlatformProcess
 	 * @return A new event, or nullptr none could be created.
 	 * @see GetSynchEventFromPool, ReturnSynchEventToPool
 	 */
-	// Message to others in the future, don't try to delete this function as it isn't exactly deprecated, but it should only ever be called from TEventPool::GetEventFromPool()
+	// Message to others in the future, don't try to delete this function as it isn't exactly deprecated,
+	// but it should only ever be called from TEventPool::GetEventFromPool()/FSafeRecyclableEvent::GetRawEvent()
 	UE_DEPRECATED(5.0, "Please use GetSynchEventFromPool to create a new event, and ReturnSynchEventToPool to release the event.")
 	static CORE_API class FEvent* CreateSynchEvent(bool bIsManualReset = false);
 
diff --git a/Engine/Source/Runtime/Core/Public/HAL/Event.h b/Engine/Source/Runtime/Core/Public/HAL/Event.h
--- a/Engine/Source/Runtime/Core/Public/HAL/Event.h
+++ b/Engine/Source/Runtime/Core/Public/HAL/Event.h
@@ -30,6 +30,8 @@ public:
 	 * @param bIsManualReset Whether the event requires manual reseting or not.
 	 * @return true if the event was created, false otherwise.
 	 */
+	// Message to others in the future, don't try to delete this function as it isn't exactly deprecated,
+	// but it should only ever be called from TEventPool::GetEventFromPool()/FSafeRecyclableEvent::GetRawEvent()
 	UE_DEPRECATED(5.0, "Direct creation of FEvent is discouraged for performance reasons. Please use FPlatformProcess::GetSynchEventFromPool/ReturnSynchEventToPool.")
 	virtual bool Create( bool bIsManualReset = false ) = 0;
 
diff --git a/Engine/Source/Runtime/Core/Public/Misc/FileHelper.h b/Engine/Source/Runtime/Core/Public/Misc/FileHelper.h
--- a/Engine/Source/Runtime/Core/Public/Misc/FileHelper.h
+++ b/Engine/Source/Runtime/Core/Public/Misc/FileHelper.h
@@ -249,20 +249,6 @@ struct FFileHelper
 	* @param OutError	If an error occurs, this is the reason why
 	*/
 	static CORE_API bool IsFilenameValidForSaving(const FString& Filename, FText& OutError);
-
-	enum class UE_DEPRECATED(5.0, "EChannelMask has been deprecated in favor of EColorChannel") EChannelMask
-	{
-		R = STRUCT_OFFSET(FColor, R),
-		G = STRUCT_OFFSET(FColor, G),
-		B = STRUCT_OFFSET(FColor, B),
-		A = STRUCT_OFFSET(FColor, A),
-		All = R | G | B | A
-	};
-
-	PRAGMA_DISABLE_DEPRECATION_WARNINGS
-	UE_DEPRECATED(5.0, "EChannelMask has been deprecated in favor of EColorChannel, please use the other CreateBitmap() method.")
-	static CORE_API bool CreateBitmap(const TCHAR* Pattern, int32 DataWidth, int32 DataHeight, const struct FColor* Data, FIntRect* SubRectangle, IFileManager* FileManager, FString* OutFilename, bool bInWriteAlpha, EChannelMask ChannelMask);
-	PRAGMA_ENABLE_DEPRECATION_WARNINGS
 };
 
 ENUM_CLASS_FLAGS(FFileHelper::EHashOptions)
diff --git a/Engine/Source/Runtime/Core/Public/Modules/ModuleManager.h b/Engine/Source/Runtime/Core/Public/Modules/ModuleManager.h
--- a/Engine/Source/Runtime/Core/Public/Modules/ModuleManager.h
+++ b/Engine/Source/Runtime/Core/Public/Modules/ModuleManager.h
@@ -32,12 +32,6 @@ class FOutputDevice;
 class UClass;
 class UWorld;
 
-#if WITH_HOT_RELOAD
-	/** If true, we are reloading a class for HotReload */
-	UE_DEPRECATED(5.0, "GIsHotReload has been deprecated, use IsReloadActive to test to see if a reload is in progress.")
-	extern CORE_API bool GIsHotReload;
-#endif
-
 #if WITH_ENGINE
 	extern CORE_API TMap<UClass*, UClass*>& GetClassesToReinstanceForHotReload(); 
 #endif
diff --git a/Engine/Source/Runtime/Core/Public/UObject/UnrealNames.h b/Engine/Source/Runtime/Core/Public/UObject/UnrealNames.h
--- a/Engine/Source/Runtime/Core/Public/UObject/UnrealNames.h
+++ b/Engine/Source/Runtime/Core/Public/UObject/UnrealNames.h
@@ -37,10 +37,4 @@ CORE_API const TCHAR* LexToString(EName Ename);
  */
 #define MAX_NETWORKED_HARDCODED_NAME 410
 
-UE_DEPRECATED(5.0, "This function is deprecated, please use ShouldReplicateAsInteger(EName Ename, const class FName& Name) instead.")
-inline bool ShouldReplicateAsInteger(EName Ename)
-{
-	return Ename <= EName(MAX_NETWORKED_HARDCODED_NAME);
-}
-
-CORE_API bool ShouldReplicateAsInteger(EName Ename, const class FName& Name);
\ No newline at end of file
+CORE_API bool ShouldReplicateAsInteger(EName Ename, const class FName& Name);
diff --git a/Engine/Source/Runtime/CoreUObject/Private/Misc/PackageName.cpp b/Engine/Source/Runtime/CoreUObject/Private/Misc/PackageName.cpp
--- a/Engine/Source/Runtime/CoreUObject/Private/Misc/PackageName.cpp
+++ b/Engine/Source/Runtime/CoreUObject/Private/Misc/PackageName.cpp
@@ -1852,21 +1852,6 @@ FName* FPackageName::FindScriptPackageName(FName InShortName)
 	return ScriptPackageNames.Find(InShortName);
 }
 
-bool FPackageName::FindPackageFileWithoutExtension(const FString& InPackageFilename, FString& OutFilename, bool InAllowTextFormats)
-{
-	bool bExists = FindPackageFileWithoutExtension(InPackageFilename, OutFilename);
-	if (!InAllowTextFormats)
-	{
-		FPackagePath PackagePath = FPackagePath::FromLocalPath(OutFilename);
-		FOpenPackageResult Result = IPackageResourceManager::Get().OpenReadPackage(PackagePath);
-		if (!Result.Archive.IsValid() || Result.Format == EPackageFormat::Text)
-		{
-			return false;
-		}
-	}
-	return bExists;
-}
-
 bool FPackageName::FindPackageFileWithoutExtension(const FString& InPackageFilename, FString& OutFilename)
 {
 	FPackagePath PackagePath = FPackagePath::FromLocalPath(InPackageFilename);
diff --git a/Engine/Source/Runtime/CoreUObject/Private/UObject/Obj.cpp b/Engine/Source/Runtime/CoreUObject/Private/UObject/Obj.cpp
--- a/Engine/Source/Runtime/CoreUObject/Private/UObject/Obj.cpp
+++ b/Engine/Source/Runtime/CoreUObject/Private/UObject/Obj.cpp
@@ -4017,11 +4017,6 @@ void UObject::UpdateSingleSectionOfConfigFile(const FString& ConfigIniName)
 	}
 }
 
-void UObject::UpdateDefaultConfigFile(const FString& SpecificFileLocation)
-{
-	UpdateSingleSectionOfConfigFile(SpecificFileLocation.IsEmpty() ? GetDefaultConfigFilename() : SpecificFileLocation);
-}
-
 bool UObject::TryUpdateDefaultConfigFile(const FString& SpecificFileLocation, bool bWarnIfFail)
 {
 	FString ConfigFile = SpecificFileLocation.IsEmpty() ? GetDefaultConfigFilename() : SpecificFileLocation;
diff --git a/Engine/Source/Runtime/CoreUObject/Private/UObject/Package.cpp b/Engine/Source/Runtime/CoreUObject/Private/UObject/Package.cpp
--- a/Engine/Source/Runtime/CoreUObject/Private/UObject/Package.cpp
+++ b/Engine/Source/Runtime/CoreUObject/Private/UObject/Package.cpp
@@ -95,10 +95,6 @@ static FAutoConsoleVariable DumpStackTraceOnPackageDirty(
 	}));
 #endif // WITH_EDITOR && !UE_BUILD_SHIPPING
 
-PRAGMA_DISABLE_DEPRECATION_WARNINGS;
-UPackage::FPreSavePackage UPackage::PreSavePackageEvent;
-UPackage::FOnPackageSaved UPackage::PackageSavedEvent;
-PRAGMA_ENABLE_DEPRECATION_WARNINGS;
 UPackage::FPreSavePackageWithContext UPackage::PreSavePackageWithContextEvent;
 UPackage::FOnPackageSavedWithContext UPackage::PackageSavedWithContextEvent;
 /** Delegate to notify subscribers when the dirty state of a package is changed.
diff --git a/Engine/Source/Runtime/CoreUObject/Private/UObject/SavePackage2.cpp b/Engine/Source/Runtime/CoreUObject/Private/UObject/SavePackage2.cpp
--- a/Engine/Source/Runtime/CoreUObject/Private/UObject/SavePackage2.cpp
+++ b/Engine/Source/Runtime/CoreUObject/Private/UObject/SavePackage2.cpp
@@ -3645,9 +3645,6 @@ void PostSavePackage(FSaveContext& SaveContext)
 	}
 
 	// Send a message that the package was saved
-	PRAGMA_DISABLE_DEPRECATION_WARNINGS;
-	UPackage::PackageSavedEvent.Broadcast(SaveContext.GetFilename(), Package);
-	PRAGMA_ENABLE_DEPRECATION_WARNINGS;
 	FObjectSaveContextData& ObjectSaveContext = SaveContext.GetObjectSaveContext();
 	ObjectSaveContext.Object = nullptr;
 	UPackage::PackageSavedWithContextEvent.Broadcast(SaveContext.GetFilename(), Package, FObjectPostSaveContext(ObjectSaveContext));
diff --git a/Engine/Source/Runtime/CoreUObject/Public/Misc/HotReloadInterface.h b/Engine/Source/Runtime/CoreUObject/Public/Misc/HotReloadInterface.h
--- a/Engine/Source/Runtime/CoreUObject/Public/Misc/HotReloadInterface.h
+++ b/Engine/Source/Runtime/CoreUObject/Public/Misc/HotReloadInterface.h
@@ -79,12 +79,6 @@ public:
 	 */
 	virtual void RequestStopCompilation() = 0;
 
-	/**
-	 * Adds a function to re-map after hot-reload.
-	 */
-	UE_DEPRECATED(5.0, "AddHotReloadFunctionRemap has been deprecated, use ReloadNotifyFunctionRemap in Reload.h instead")
-	virtual void AddHotReloadFunctionRemap(FNativeFuncPtr NewFunctionPointer, FNativeFuncPtr OldFunctionPointer) = 0;
-
 	/**
 	 * Performs hot reload from the editor of all currently loaded game modules.
 	 *
@@ -105,14 +99,6 @@ public:
 	 */
 	virtual ECompilationResult::Type RebindPackages(const TArray<UPackage*>& Packages, EHotReloadFlags Flags, FOutputDevice &Ar) = 0;
 
-	/** Called when a Hot Reload event has completed. 
-	 * 
-	 * @param	bWasTriggeredAutomatically	True if the hot reload was invoked automatically by the hot reload system after detecting a changed DLL
-	 */
-	DECLARE_MULTICAST_DELEGATE_OneParam(FHotReloadEvent, bool /* bWasTriggeredAutomatically */ );
-	UE_DEPRECATED(5.0, "OnHotReload has been deprecated, use FCoreUObjectDelegates::ReloadCompleteDelegate.")
-	virtual FHotReloadEvent& OnHotReload() = 0;
-
 	/**
 	 * Gets an event delegate that is executed when compilation of a module has started.
 	 *
diff --git a/Engine/Source/Runtime/CoreUObject/Public/Misc/PackageName.h b/Engine/Source/Runtime/CoreUObject/Public/Misc/PackageName.h
--- a/Engine/Source/Runtime/CoreUObject/Public/Misc/PackageName.h
+++ b/Engine/Source/Runtime/CoreUObject/Public/Misc/PackageName.h
@@ -355,22 +355,6 @@ public:
 	 */
 	static COREUOBJECT_API FString GetContentPathForPackageRoot(FStringView InPackageRoot);
 
-	/**
-	 * Checks if the package exists on disk.
-	 * 
-	 * @param LongPackageName Package name.
-	 * @param Guid If nonnull, and the package is found on disk but does not have this PackageGuid in its FPackageFileSummary::Guid, false is returned
-	 * @param OutFilename Package filename on disk.
-	 * @param InAllowTextFormats Detect text format packages as well as binary (priority to text)
-	 * @return true if the specified package name points to an existing package, false otherwise.
-	 **/
-	UE_DEPRECATED(5.0, "Deprecated. UPackage::Guid has not been used by the engine for a long time. Call DoesPackageExist without a Guid.")
-	static bool DoesPackageExist(const FString& LongPackageName, const FGuid* Guid, FString* OutFilename, bool InAllowTextFormats = true)
-	{
-		return DoesPackageExist(LongPackageName, OutFilename, InAllowTextFormats);
-	}
-
-
 	/**
 	 * Checks if the package exists on disk.
 	 * 
@@ -940,9 +924,6 @@ public:
 	*/
 	static COREUOBJECT_API bool FindPackageFileWithoutExtension(const FString& InPackageFilename, FString& OutFilename);
 
-	UE_DEPRECATED(5.0, "Specifying AllowTextFormats is no longer supported. Text format is instead specified as a field on the result struct from IPackageResourceManager->OpenReadPackage.")
-	static COREUOBJECT_API bool FindPackageFileWithoutExtension(const FString& InPackageFilename, FString& OutFilename, bool InAllowTextFormats);
-
 	/**
 	 * Converts a long package name to the case it exists as on disk.
 	 *
diff --git a/Engine/Source/Runtime/CoreUObject/Public/Serialization/ArchiveReplaceObjectAndStructPropertyRef.h b/Engine/Source/Runtime/CoreUObject/Public/Serialization/ArchiveReplaceObjectAndStructPropertyRef.h
--- a/Engine/Source/Runtime/CoreUObject/Public/Serialization/ArchiveReplaceObjectAndStructPropertyRef.h
+++ b/Engine/Source/Runtime/CoreUObject/Public/Serialization/ArchiveReplaceObjectAndStructPropertyRef.h
@@ -37,27 +37,6 @@ public:
 		}
 	}
 
-	UE_DEPRECATED(5.0, "Use version that supplies flags via enum.")
-	FArchiveReplaceObjectAndStructPropertyRef
-	(
-		UObject* InSearchObject,
-		const TMap<T*, T*>& InReplacementMap,
-		bool bNullPrivateRefs,
-		bool bIgnoreOuterRef,
-		bool bIgnoreArchetypeRef,
-		bool bDelayStart = false,
-		bool bIgnoreClassGeneratedByRef = true
-	)
-		: FArchiveReplaceObjectRef<T>(InSearchObject, InReplacementMap, bNullPrivateRefs, bIgnoreOuterRef, bIgnoreArchetypeRef, true, bIgnoreClassGeneratedByRef)
-	{
-		// Note: We intentionally request a delayed start above; otherwise the base class ctor will use the base archive type for serialization, and our overrides won't get called.
-		// Here we're looking for whether or not the *caller* of this ctor has asked for a delayed start; in that case, we defer to the caller as to when our serialization will start.
-		if (!bDelayStart)
-		{
-			FArchiveReplaceObjectRef<T>::SerializeSearchObject();
-		}
-	}
-
 	/**
 	 * Serializes the reference to FProperties
 	 */
diff --git a/Engine/Source/Runtime/CoreUObject/Public/Serialization/ArchiveReplaceObjectRef.h b/Engine/Source/Runtime/CoreUObject/Public/Serialization/ArchiveReplaceObjectRef.h
--- a/Engine/Source/Runtime/CoreUObject/Public/Serialization/ArchiveReplaceObjectRef.h
+++ b/Engine/Source/Runtime/CoreUObject/Public/Serialization/ArchiveReplaceObjectRef.h
@@ -196,26 +196,6 @@ public:
 		}
 	}
 
-	UE_DEPRECATED(5.0, "Use version that supplies flags via enum.")
-	FArchiveReplaceObjectRef
-	(
-		UObject* InSearchObject,
-		const TMap<T*,T*>& inReplacementMap,
-		bool bNullPrivateRefs,
-		bool bIgnoreOuterRef,
-		bool bIgnoreArchetypeRef,
-		bool bDelayStart = false,
-		bool bIgnoreClassGeneratedByRef = true
-	)
-	: FArchiveReplaceObjectRef(InSearchObject, inReplacementMap,
-		  (bNullPrivateRefs ? EArchiveReplaceObjectFlags::NullPrivateRefs : EArchiveReplaceObjectFlags::None)
-		| (bIgnoreOuterRef ? EArchiveReplaceObjectFlags::IgnoreOuterRef : EArchiveReplaceObjectFlags::None)
-		| (bIgnoreArchetypeRef ? EArchiveReplaceObjectFlags::IgnoreArchetypeRef : EArchiveReplaceObjectFlags::None)
-		| (bDelayStart ? EArchiveReplaceObjectFlags::DelayStart : EArchiveReplaceObjectFlags::None)
-		| (!bIgnoreClassGeneratedByRef ? EArchiveReplaceObjectFlags::IncludeClassGeneratedByRef : EArchiveReplaceObjectFlags::None))
-	{
-	}
-
 	/**
 	 * Starts the serialization of the root object
 	 */
diff --git a/Engine/Source/Runtime/CoreUObject/Public/Serialization/ArchiveReplaceOrClearExternalReferences.h b/Engine/Source/Runtime/CoreUObject/Public/Serialization/ArchiveReplaceOrClearExternalReferences.h
--- a/Engine/Source/Runtime/CoreUObject/Public/Serialization/ArchiveReplaceOrClearExternalReferences.h
+++ b/Engine/Source/Runtime/CoreUObject/Public/Serialization/ArchiveReplaceOrClearExternalReferences.h
@@ -34,21 +34,6 @@ public:
 		}
 	}
 
-	UE_DEPRECATED(5.0, "Use version that passes Flags instead of just bDelayStart")
-	FArchiveReplaceOrClearExternalReferences
-		( UObject* InSearchObject
-		, const TMap<T*, T*>& InReplacementMap
-		, UPackage* InDestPackage
-		, bool bDelayStart)
-		: TSuper(InSearchObject, InReplacementMap, EArchiveReplaceObjectFlags::DelayStart)
-		, DestPackage(InDestPackage)
-	{
-		if (!bDelayStart)
-		{
-			this->SerializeSearchObject();
-		}
-	}
-
 	FArchive& operator<<(UObject*& Obj)
 	{
 		UObject* Resolved = Obj;
diff --git a/Engine/Source/Runtime/CoreUObject/Public/UObject/Field.h b/Engine/Source/Runtime/CoreUObject/Public/UObject/Field.h
--- a/Engine/Source/Runtime/CoreUObject/Public/UObject/Field.h
+++ b/Engine/Source/Runtime/CoreUObject/Public/UObject/Field.h
@@ -795,18 +795,6 @@ public:
 		}
 	}
 
-	UE_DEPRECATED(5.0, "HasAnyCastFlags is deprecated. Not all FField has CastFlag. Use IsA instead.")
-	inline bool HasAnyCastFlags(const uint64 InCastFlags) const
-	{
-		return !!(GetCastFlags() & InCastFlags);
-	}
-
-	UE_DEPRECATED(5.0, "HasAllCastFlags is deprecated. Not all FField has CastFlag. Use IsA instead.")
-	inline bool HasAllCastFlags(const uint64 InCastFlags) const
-	{
-		return (GetCastFlags() & InCastFlags) == InCastFlags;
-	}
-
 	void AppendName(FString& ResultString) const
 	{
 		GetFName().AppendString(ResultString);
diff --git a/Engine/Source/Runtime/CoreUObject/Public/UObject/Object.h b/Engine/Source/Runtime/CoreUObject/Public/UObject/Object.h
--- a/Engine/Source/Runtime/CoreUObject/Public/UObject/Object.h
+++ b/Engine/Source/Runtime/CoreUObject/Public/UObject/Object.h
@@ -1335,12 +1335,6 @@ public:
 	};
 	COREUOBJECT_API void SaveConfig(const FSaveConfigContext& Context);
 
-	/**
-	 * Saves just the section(s) for this class into the default ini file for the class (with just the changes from base)
-	 */
-	UE_DEPRECATED(5.0, "TryUpdateDefaultConfigFile replaces UpdateDefaultConfigFile")
-	COREUOBJECT_API void UpdateDefaultConfigFile(const FString& SpecificFileLocation = "");
-
 	/**
 	 * Try to Saves just the section(s) for this class into the default ini file for the class (with just the changes from base)
 	 *
diff --git a/Engine/Source/Runtime/CoreUObject/Public/UObject/ObjectMacros.h b/Engine/Source/Runtime/CoreUObject/Public/UObject/ObjectMacros.h
--- a/Engine/Source/Runtime/CoreUObject/Public/UObject/ObjectMacros.h
+++ b/Engine/Source/Runtime/CoreUObject/Public/UObject/ObjectMacros.h
@@ -2491,19 +2491,6 @@ namespace UE
 	};
 }
 
-namespace UE4
-{
-	using ELoadConfigPropagationFlags UE_DEPRECATED(5.0, "Use UE namespace instead of UE4") = UE::ELoadConfigPropagationFlags;
-
-	UE_DEPRECATED(5.0, "Use UE namespace instead of UE4") inline constexpr UE::ELoadConfigPropagationFlags LCPF_None							= UE::ELoadConfigPropagationFlags::LCPF_None;
-	UE_DEPRECATED(5.0, "Use UE namespace instead of UE4") inline constexpr UE::ELoadConfigPropagationFlags LCPF_ReadParentSections				= UE::ELoadConfigPropagationFlags::LCPF_ReadParentSections;
-	UE_DEPRECATED(5.0, "Use UE namespace instead of UE4") inline constexpr UE::ELoadConfigPropagationFlags LCPF_PropagateToChildDefaultObjects	= UE::ELoadConfigPropagationFlags::LCPF_PropagateToChildDefaultObjects;
-	UE_DEPRECATED(5.0, "Use UE namespace instead of UE4") inline constexpr UE::ELoadConfigPropagationFlags LCPF_PropagateToInstances			= UE::ELoadConfigPropagationFlags::LCPF_PropagateToInstances;
-	UE_DEPRECATED(5.0, "Use UE namespace instead of UE4") inline constexpr UE::ELoadConfigPropagationFlags LCPF_ReloadingConfigData				= UE::ELoadConfigPropagationFlags::LCPF_ReloadingConfigData;
-	UE_DEPRECATED(5.0, "Use UE namespace instead of UE4") inline constexpr UE::ELoadConfigPropagationFlags LCPF_PersistentFlags					= UE::ELoadConfigPropagationFlags::LCPF_PersistentFlags;
-}
-
-
 /**
  * Helper class used to save and restore information across a StaticAllocateObject over the top of an existing object.
  * Currently only used by UClass
diff --git a/Engine/Source/Runtime/CoreUObject/Public/UObject/Package.h b/Engine/Source/Runtime/CoreUObject/Public/UObject/Package.h
--- a/Engine/Source/Runtime/CoreUObject/Public/UObject/Package.h
+++ b/Engine/Source/Runtime/CoreUObject/Public/UObject/Package.h
@@ -282,12 +282,8 @@ public:
 	DECLARE_MULTICAST_DELEGATE_FourParams(FOnMetaDataTransferRequestedDelegate, const FMetaData& OldMetaData, FMetaData& NewMetaData, const UPackage* OldPackage, UPackage* NewPackage);
 #endif
 	
-	UE_DEPRECATED(5.0, "Use PreSavePackageWithContextEvent instead.")
-	COREUOBJECT_API static FPreSavePackage PreSavePackageEvent;
 	/** Delegate to notify subscribers when a package is about to be saved. */
 	COREUOBJECT_API static FPreSavePackageWithContext PreSavePackageWithContextEvent;
-	UE_DEPRECATED(5.0, "Use PackageSavedWithContextEvent instead.")
-	COREUOBJECT_API static FOnPackageSaved PackageSavedEvent;
 	/** Delegate to notify subscribers when a package has been saved. This is triggered when the package saving
 	*  has completed and was successful. */
 	COREUOBJECT_API static FOnPackageSavedWithContext PackageSavedWithContextEvent;
diff --git a/Engine/Source/Runtime/CoreUObject/Public/UObject/PackageFileSummary.h b/Engine/Source/Runtime/CoreUObject/Public/UObject/PackageFileSummary.h
--- a/Engine/Source/Runtime/CoreUObject/Public/UObject/PackageFileSummary.h
+++ b/Engine/Source/Runtime/CoreUObject/Public/UObject/PackageFileSummary.h
@@ -303,23 +303,11 @@ public:
 		return FileVersionUE;
 	}
 
-	UE_DEPRECATED(5.0, "Use UEVer instead which returns the version as a FPackageFileVersion. See the @FPackageFileVersion documentation for further details")
-	int32 GetFileVersionUE4() const
-	{
-		return FileVersionUE.FileVersionUE4;
-	}
-
 	int32 GetFileVersionLicenseeUE() const
 	{
 		return FileVersionLicenseeUE;
 	}
 
-	UE_DEPRECATED(5.0, "Use GetFileVersionLicenseeUE instead")
-	int32 GetFileVersionLicenseeUE4() const
-	{
-		return GetFileVersionLicenseeUE();
-	}
-
 	const FCustomVersionContainer& GetCustomVersionContainer() const
 	{
 		return CustomVersionContainer;
@@ -330,12 +318,6 @@ public:
 	/** Used to manually set the package file and licensee versions */
 	COREUOBJECT_API void SetFileVersions(const int32 EpicUE4, const int32 EpicUE5, const int32 LicenseeUE, const bool bInSaveUnversioned = false);
 	
-	UE_DEPRECATED(5.0, "Use the other overload of SetFileVersions that takes an UE5 version as well")
-	inline void SetFileVersions(const int32 EpicUE, const int32 LicenseeUE, const bool bInSaveUnversioned = false)
-	{
-		SetFileVersions(EpicUE, GPackageFileUEVersion.FileVersionUE5, LicenseeUE, bInSaveUnversioned);
-	}
-
 	/** Set both the package file versions and the licensee version to the most recent version numbers supported */
 	void SetToLatestFileVersions(const bool bInSaveUnversioned)
 	{
diff --git a/Engine/Source/Runtime/CoreUObject/Public/UObject/ReferenceChainSearch.h b/Engine/Source/Runtime/CoreUObject/Public/UObject/ReferenceChainSearch.h
--- a/Engine/Source/Runtime/CoreUObject/Public/UObject/ReferenceChainSearch.h
+++ b/Engine/Source/Runtime/CoreUObject/Public/UObject/ReferenceChainSearch.h
@@ -164,9 +164,6 @@ public:
 	/** Single node in the reference graph */
 	struct FGraphNode
 	{
-		/** Object pointer */
-		UE_DEPRECATED(5.0, "Direct Object reference has been deprecated. Use ObjectInfo member variable instead.")
-		UObject* Object = nullptr;
 		/** Object pointer */
 		FGCObjectInfo* ObjectInfo = nullptr;
 		/** Objects referenced by this object with reference info */
diff --git a/Engine/Source/Runtime/CoreUObject/Public/UObject/UObjectBaseUtility.h b/Engine/Source/Runtime/CoreUObject/Public/UObject/UObjectBaseUtility.h
--- a/Engine/Source/Runtime/CoreUObject/Public/UObject/UObjectBaseUtility.h
+++ b/Engine/Source/Runtime/CoreUObject/Public/UObject/UObjectBaseUtility.h
@@ -753,15 +753,6 @@ public:
 	 */
 	COREUOBJECT_API FPackageFileVersion GetLinkerUEVersion() const;
 		
-	UE_DEPRECATED(5.0, "Use GetLinkerUEVersion instead which returns the version as a FPackageFileVersion. See the @FPackageFileVersion documentation for further details")
-	inline int32 GetLinkerUE4Version() const 
-	{ 
-		// Existing code calling GetLinkerUE4Version will be testing against UE4 version numbers so 
-		// we can just return the UE4 version number.
-		// All new code that might actually need the UE5 version as well should be calling ::GetLinkerUEVersion
-		return GetLinkerUEVersion().FileVersionUE4;
-	}
-
 	/**
 	 * Returns the licensee version of the linker for this object.
 	 *
@@ -773,10 +764,6 @@ public:
 	 */
 	COREUOBJECT_API int32 GetLinkerLicenseeUEVersion() const;
 
-	UE_DEPRECATED(5.0, "Use GetLinkerLicenseeUEVersion instead")
-	inline int32 GetLinkerLicenseeUE4Version() const { return GetLinkerLicenseeUEVersion(); }
-
-
 	/**
 	 * Returns the custom version of the linker for this object corresponding to the given custom version key.
 	 *
